Webm: audio synchronization fixes

diff --git a/Engine/Plugins/Media/WebMMedia/Source/WebMMedia/Private/Player/MkvFileReader.cpp b/Engine/Plugins/Media/WebMMedia/Source/WebMMedia/Private/Player/MkvFileReader.cpp
--- a/Engine/Plugins/Media/WebMMedia/Source/WebMMedia/Private/Player/MkvFileReader.cpp
+++ b/Engine/Plugins/Media/WebMMedia/Source/WebMMedia/Private/Player/MkvFileReader.cpp
@@ -14,6 +14,19 @@ bool FMkvFileReader::Open(const TCHAR* Filename)
 	return m_file.IsValid();
 }
 
+FTimespan FMkvFileReader::GetVideoFrameDuration(const mkvparser::VideoTrack& track)
+{
+	double FrameRate = track.GetFrameRate();
+	if (FrameRate > 0)
+	{
+		return FTimespan::FromSeconds(1.0 / FrameRate);
+	}
+	else
+	{
+		return FTimespan::FromSeconds(1.0 / 30.0);
+	}
+}
+
 int FMkvFileReader::Read(long long Position, long Lenght, unsigned char* Buffer)
 {
 	check(m_file.IsValid());
diff --git a/Engine/Plugins/Media/WebMMedia/Source/WebMMedia/Private/Player/MkvFileReader.h b/Engine/Plugins/Media/WebMMedia/Source/WebMMedia/Private/Player/MkvFileReader.h
--- a/Engine/Plugins/Media/WebMMedia/Source/WebMMedia/Private/Player/MkvFileReader.h
+++ b/Engine/Plugins/Media/WebMMedia/Source/WebMMedia/Private/Player/MkvFileReader.h
@@ -17,6 +17,8 @@ class FMkvFileReader : public mkvparser::IMkvReader
 public:
 	bool Open(const TCHAR* Filename);
 
+	static FTimespan GetVideoFrameDuration(const mkvparser::VideoTrack& track);
+
 public:
 	// IMkvReader interface
 	virtual int Read(long long Position, long Lenght, unsigned char* Buffer) override;
diff --git a/Engine/Plugins/Media/WebMMedia/Source/WebMMedia/Private/Player/WebMAudioDecoder.cpp b/Engine/Plugins/Media/WebMMedia/Source/WebMMedia/Private/Player/WebMAudioDecoder.cpp
--- a/Engine/Plugins/Media/WebMMedia/Source/WebMMedia/Private/Player/WebMAudioDecoder.cpp
+++ b/Engine/Plugins/Media/WebMMedia/Source/WebMMedia/Private/Player/WebMAudioDecoder.cpp
@@ -4,7 +4,7 @@
 #include "WebMMediaPrivate.h"
 #include "WebMMediaFrame.h"
 #include "WebMMediaAudioSample.h"
-#include "MediaSamples.h"
+#include "WebMSamplesSink.h"
 
 THIRD_PARTY_INCLUDES_START
 #include <opus.h>
@@ -18,9 +18,9 @@ struct FWebMAudioDecoder::FVorbisDecoder
 	vorbis_block Block;
 };
 
-FWebMAudioDecoder::FWebMAudioDecoder(TSharedPtr<FMediaSamples, ESPMode::ThreadSafe> InSamples)
-	: Samples(InSamples)
-	, AudioSamplePool(new FWebMMediaAudioSamplePool)
+FWebMAudioDecoder::FWebMAudioDecoder(IWebMSamplesSink& InSamples)
+	: AudioSamplePool(new FWebMMediaAudioSamplePool)
+	, Samples(InSamples)
 	, OpusDecoder(nullptr)
 {
 }
@@ -94,7 +94,7 @@ bool FWebMAudioDecoder::InitializeVorbis(const uint8* CodecPrivateData, size_t C
 {
 	if (CodecPrivateDataSize < 3 || !CodecPrivateData || CodecPrivateData[0] != 2)
 	{
-		UE_LOG(LogWebMMedia, Display, TEXT("Failed to initialize Vorbis - invalid data"));
+		UE_LOG(LogWebMMedia, Warning, TEXT("Failed to initialize Vorbis - invalid data"));
 		return false;
 	}
 
@@ -108,7 +108,7 @@ bool FWebMAudioDecoder::InitializeVorbis(const uint8* CodecPrivateData, size_t C
 		{
 			if (Offset >= CodecPrivateDataSize)
 			{
-				UE_LOG(LogWebMMedia, Display, TEXT("Failed to initialize Vorbis - invalid offset"));
+				UE_LOG(LogWebMMedia, Warning, TEXT("Failed to initialize Vorbis - invalid offset"));
 				return false;
 			}
 			HeaderSize[i] += CodecPrivateData[Offset];
@@ -122,7 +122,7 @@ bool FWebMAudioDecoder::InitializeVorbis(const uint8* CodecPrivateData, size_t C
 
 	if (HeaderSize[0] + HeaderSize[1] + HeaderSize[2] + Offset != CodecPrivateDataSize)
 	{
-		UE_LOG(LogWebMMedia, Display, TEXT("Failed to initialize Vorbis - parameters mismatch"));
+		UE_LOG(LogWebMMedia, Warning, TEXT("Failed to initialize Vorbis - parameters mismatch"));
 		return false;
 	}
 
@@ -150,7 +150,7 @@ bool FWebMAudioDecoder::InitializeVorbis(const uint8* CodecPrivateData, size_t C
 	{
 		if (vorbis_synthesis_headerin(&VorbisDecoder->Info, &Vc, &Packet[i]))
 		{
-			UE_LOG(LogWebMMedia, Display, TEXT("Failed to initialize Vorbis - invalid headers"));
+			UE_LOG(LogWebMMedia, Warning, TEXT("Failed to initialize Vorbis - invalid headers"));
 			vorbis_comment_clear(&Vc);
 			return false;
 		}
@@ -160,19 +160,19 @@ bool FWebMAudioDecoder::InitializeVorbis(const uint8* CodecPrivateData, size_t C
 
 	if (vorbis_synthesis_init(&VorbisDecoder->DspState, &VorbisDecoder->Info))
 	{
-		UE_LOG(LogWebMMedia, Display, TEXT("Failed to initialize Vorbis (synthesis)"));
+		UE_LOG(LogWebMMedia, Warning, TEXT("Failed to initialize Vorbis (synthesis)"));
 		return false;
 	}
 
 	if (VorbisDecoder->Info.channels != Channels || VorbisDecoder->Info.rate != SampleRate)
 	{
-		UE_LOG(LogWebMMedia, Display, TEXT("Failed to initialize Vorbis - invalid parameters"));
+		UE_LOG(LogWebMMedia, Warning, TEXT("Failed to initialize Vorbis - invalid parameters"));
 		return false;
 	}
 
 	if (vorbis_block_init(&VorbisDecoder->DspState, &VorbisDecoder->Block))
 	{
-		UE_LOG(LogWebMMedia, Display, TEXT("Failed to initialize Vorbis"));
+		UE_LOG(LogWebMMedia, Warning, TEXT("Failed to initialize Vorbis"));
 		return false;
 	}
 
@@ -219,8 +219,10 @@ void FWebMAudioDecoder::DoDecodeAudioFrames(const TArray<TSharedPtr<FWebMFrame>>
 		if (NumOfSamplesDecoded > 0)
 		{
 			TSharedRef<FWebMMediaAudioSample, ESPMode::ThreadSafe> AudioSample = AudioSamplePool->AcquireShared();
-			AudioSample->Initialize(DecodeBuffer.GetData(), NumOfSamplesDecoded * Channels * 2, Channels, SampleRate, AudioFrame->Time, FTimespan::FromHours(1));
-			Samples->AddAudio(AudioSample);
+			FTimespan Duration = (NumOfSamplesDecoded * ETimespan::TicksPerSecond) / SampleRate;
+			AudioSample->Initialize(DecodeBuffer.GetData(), NumOfSamplesDecoded * Channels * 2, Channels, SampleRate, AudioFrame->Time, Duration);
+
+			Samples.AddAudioSampleFromDecodingThread(AudioSample);
 		}
 	}
 }
@@ -230,7 +232,7 @@ int32 FWebMAudioDecoder::DecodeOpus(const TSharedPtr<FWebMFrame>& AudioFrame)
 	int32 Result = opus_decode(OpusDecoder, AudioFrame->Data.GetData(), AudioFrame->Data.Num(), (opus_int16*) DecodeBuffer.GetData(), FrameSize, 0);
 	if (Result < 0)
 	{
-		UE_LOG(LogWebMMedia, Display, TEXT("Error decoding Opus audio frame"));
+		UE_LOG(LogWebMMedia, Warning, TEXT("Error decoding Opus audio frame"));
 	}
 
 	return Result;
@@ -246,13 +248,13 @@ int32 FWebMAudioDecoder::DecodeVorbis(const TSharedPtr<FWebMFrame>& AudioFrame)
 
 	if (vorbis_synthesis(&VorbisDecoder->Block, &Packet))
 	{
-		UE_LOG(LogWebMMedia, Display, TEXT("Error decoding Vorbis audio frame - vorbis_synthesis failed"));
+		UE_LOG(LogWebMMedia, Warning, TEXT("Error decoding Vorbis audio frame - vorbis_synthesis failed"));
 		return -1;
 	}
 
 	if (vorbis_synthesis_blockin(&VorbisDecoder->DspState, &VorbisDecoder->Block))
 	{
-		UE_LOG(LogWebMMedia, Display, TEXT("Error decoding Vorbis audio frame - vorbis_synthesis_blockin failed"));
+		UE_LOG(LogWebMMedia, Warning, TEXT("Error decoding Vorbis audio frame - vorbis_synthesis_blockin failed"));
 		return - 1;
 	}
 
diff --git a/Engine/Plugins/Media/WebMMedia/Source/WebMMedia/Private/Player/WebMMediaPlayer.cpp b/Engine/Plugins/Media/WebMMedia/Source/WebMMedia/Private/Player/WebMMediaPlayer.cpp
--- a/Engine/Plugins/Media/WebMMedia/Source/WebMMedia/Private/Player/WebMMediaPlayer.cpp
+++ b/Engine/Plugins/Media/WebMMedia/Source/WebMMedia/Private/Player/WebMMediaPlayer.cpp
@@ -12,6 +12,8 @@
 #include "WebMVideoDecoder.h"
 #include "WebMAudioDecoder.h"
 #include "WebMMediaFrame.h"
+#include "WebMMediaTextureSample.h"
+#include "WebMMediaAudioSample.h"
 
 #define LOCTEXT_NAMESPACE "FWebMMediaPlayer"
 
@@ -142,12 +144,14 @@ bool FWebMMediaPlayer::Open(const FString& Url, const IMediaOptions* /*Options*/
 
 	if (!MkvRead())
 	{
+		MkvReader.Reset();
+
 		UE_LOG(LogWebMMedia, Error, TEXT("Error parsing matroska file: %s"), *FilePath);
 		return false;
 	}
 
-	VideoDecoder.Reset(new FWebMVideoDecoder(Samples));
-	AudioDecoder.Reset(new FWebMAudioDecoder(Samples));
+	VideoDecoder.Reset(new FWebMVideoDecoder(*this));
+	AudioDecoder.Reset(new FWebMAudioDecoder(*this));
 
 	if (SelectedAudioTrack != INDEX_NONE)
 	{
@@ -218,6 +222,12 @@ void FWebMMediaPlayer::TickFetch(FTimespan DeltaTime, FTimespan Timecode)
 	TArray<TSharedPtr<FWebMFrame>> VideoFrames;
 	TArray<TSharedPtr<FWebMFrame>> AudioFrames;
 
+	if (CurrentReadTime < CurrentTime - ReadBufferLength / 2)
+	{
+		UE_LOG(LogWebMMedia, Warning, TEXT("Playback is very behind, try to catchup and reset sync"));
+		Samples->FlushSamples();
+	}
+
 	// Read frames up to 1 secs in the future
 	while (CurrentReadTime < CurrentTime + ReadBufferLength)
 	{
@@ -243,11 +253,13 @@ void FWebMMediaPlayer::TickFetch(FTimespan DeltaTime, FTimespan Timecode)
 
 			if (SelectedVideoTrack != INDEX_NONE && VideoTracks[SelectedVideoTrack]->GetNumber() == TrackNumber)
 			{
+				Frame->Duration = FMkvFileReader::GetVideoFrameDuration(*VideoTracks[SelectedVideoTrack]);
 				VideoFrames.Add(Frame);
-				CurrentReadTime += FTimespan::FromSeconds(1.0 / VideoTracks[SelectedVideoTrack]->GetFrameRate());
+				CurrentReadTime += Frame->Duration;
 			}
 			else if (SelectedAudioTrack != INDEX_NONE && AudioTracks[SelectedAudioTrack]->GetNumber() == TrackNumber)
 			{
+				// We will set duration after decompression.
 				AudioFrames.Add(Frame);
 			}
 		}
@@ -748,6 +760,28 @@ void FWebMMediaPlayer::Stop()
 	Pause();
 }
 
+void FWebMMediaPlayer::AddVideoSampleFromDecodingThread(TSharedRef<FWebMMediaTextureSample, ESPMode::ThreadSafe> Sample)
+{
+	if (Sample->GetTime() < CurrentTime)
+	{
+		// We don't care about expired samples
+		return;
+	}
+
+	Samples->AddVideo(Sample);
+}
+
+void FWebMMediaPlayer::AddAudioSampleFromDecodingThread(TSharedRef<FWebMMediaAudioSample, ESPMode::ThreadSafe> Sample)
+{
+	if (Sample->GetTime() < CurrentTime)
+	{
+		// We don't care about expired samples
+		return;
+	}
+
+	Samples->AddAudio(Sample);
+}
+
 #undef LOCTEXT_NAMESPACE
 
 #endif // WITH_WEBM_LIBS
diff --git a/Engine/Plugins/Media/WebMMedia/Source/WebMMedia/Private/Player/WebMMediaPlayer.h b/Engine/Plugins/Media/WebMMedia/Source/WebMMedia/Private/Player/WebMMediaPlayer.h
--- a/Engine/Plugins/Media/WebMMedia/Source/WebMMedia/Private/Player/WebMMediaPlayer.h
+++ b/Engine/Plugins/Media/WebMMedia/Source/WebMMedia/Private/Player/WebMMediaPlayer.h
@@ -15,8 +15,8 @@
 #include "IMediaView.h"
 #include "IMediaEventSink.h"
 #include "Misc/Timespan.h"
-#include "Templates/SharedPointer.h"
 #include "MkvFileReader.h"
+#include "WebMSamplesSink.h"
 
 class FMediaSamples;
 class FWebMVideoDecoder;
@@ -24,6 +24,7 @@ class FWebMAudioDecoder;
 
 class FWebMMediaPlayer
 	: public IMediaPlayer
+	, public IWebMSamplesSink
 	, protected IMediaCache
 	, protected IMediaControls
 	, protected IMediaTracks
@@ -33,6 +34,11 @@ public:
 	FWebMMediaPlayer(IMediaEventSink& InEventSink);
 	virtual ~FWebMMediaPlayer();
 
+public:
+	//~ IWebMSamplesSink interface
+	virtual void AddVideoSampleFromDecodingThread(TSharedRef<FWebMMediaTextureSample, ESPMode::ThreadSafe> Sample) override;
+	virtual void AddAudioSampleFromDecodingThread(TSharedRef<FWebMMediaAudioSample, ESPMode::ThreadSafe> Sample) override;
+
 public:
 	//~ IMediaPlayer interface
 	virtual void Close() override;
diff --git a/Engine/Plugins/Media/WebMMedia/Source/WebMMedia/Private/Player/WebMMediaTextureSample.cpp b/Engine/Plugins/Media/WebMMedia/Source/WebMMedia/Private/Player/WebMMediaTextureSample.cpp
--- a/Engine/Plugins/Media/WebMMedia/Source/WebMMedia/Private/Player/WebMMediaTextureSample.cpp
+++ b/Engine/Plugins/Media/WebMMedia/Source/WebMMedia/Private/Player/WebMMediaTextureSample.cpp
@@ -8,11 +8,12 @@ FWebMMediaTextureSample::FWebMMediaTextureSample()
 {
 }
 
-void FWebMMediaTextureSample::Initialize(FIntPoint InDisplaySize, FIntPoint InTotalSize, FTimespan InTime)
+void FWebMMediaTextureSample::Initialize(FIntPoint InDisplaySize, FIntPoint InTotalSize, FTimespan InTime, FTimespan InDuration)
 {
 	Time = InTime;
 	DisplaySize = InDisplaySize;
 	TotalSize = InTotalSize;
+	Duration = InDuration;
 }
 
 void FWebMMediaTextureSample::CreateTexture()
@@ -50,7 +51,7 @@ FIntPoint FWebMMediaTextureSample::GetDim() const
 
 FTimespan FWebMMediaTextureSample::GetDuration() const
 {
-	return FTimespan(ETimespan::TicksPerSecond / 30);
+	return Duration;
 }
 
 EMediaTextureSampleFormat FWebMMediaTextureSample::GetFormat() const
diff --git a/Engine/Plugins/Media/WebMMedia/Source/WebMMedia/Private/Player/WebMVideoDecoder.cpp b/Engine/Plugins/Media/WebMMedia/Source/WebMMedia/Private/Player/WebMVideoDecoder.cpp
--- a/Engine/Plugins/Media/WebMMedia/Source/WebMMedia/Private/Player/WebMVideoDecoder.cpp
+++ b/Engine/Plugins/Media/WebMMedia/Source/WebMMedia/Private/Player/WebMVideoDecoder.cpp
@@ -7,8 +7,8 @@
 #include "WebMMediaPrivate.h"
 #include "WebMMediaFrame.h"
 #include "WebMMediaTextureSample.h"
+#include "WebMSamplesSink.h"
 #include "MediaShaders.h"
-#include "MediaSamples.h"
 #include "PipelineStateCache.h"
 #include "RHIStaticStates.h"
 #include "Containers/DynamicRHIResourceArray.h"
@@ -61,9 +61,9 @@ namespace
 	TGlobalResource<FMoviePlaybackResources> GMoviePlayerResources;
 }
 
-FWebMVideoDecoder::FWebMVideoDecoder(TSharedPtr<FMediaSamples, ESPMode::ThreadSafe> InSamples)
-	: Samples(InSamples)
-	, VideoSamplePool(new FWebMMediaTextureSamplePool)
+FWebMVideoDecoder::FWebMVideoDecoder(IWebMSamplesSink& InSamples)
+	: VideoSamplePool(new FWebMMediaTextureSamplePool)
+	, Samples(InSamples)
 	, bTexturesCreated(false)
 	, bIsInitialized(false)
 {
@@ -147,7 +147,7 @@ void FWebMVideoDecoder::DoDecodeVideoFrames(const TArray<TSharedPtr<FWebMFrame>>
 
 			TSharedRef<FWebMMediaTextureSample, ESPMode::ThreadSafe> VideoSample = VideoSamplePool->AcquireShared();
 
-			VideoSample->Initialize(FIntPoint(Image->d_w, Image->d_h), FIntPoint(Image->d_w, Image->d_h), VideoFrame->Time);
+			VideoSample->Initialize(FIntPoint(Image->d_w, Image->d_h), FIntPoint(Image->d_w, Image->d_h), VideoFrame->Time, VideoFrame->Duration);
 
 			FConvertParams Params;
 			Params.VideoSample = VideoSample;
@@ -278,7 +278,7 @@ void FWebMVideoDecoder::ConvertYUVToRGBAndSubmit(const FConvertParams& Params)
 		CommandList.EndRenderPass();
 		CommandList.CopyToResolveTarget(RenderTarget, RenderTarget, FResolveParams());
 
-		Samples->AddVideo(VideoSample.ToSharedRef());
+		Samples.AddVideoSampleFromDecodingThread(VideoSample.ToSharedRef());
 	}
 }
 
diff --git a/Engine/Plugins/Media/WebMMedia/Source/WebMMedia/Private/WebMContainer.cpp b/Engine/Plugins/Media/WebMMedia/Source/WebMMedia/Private/WebMContainer.cpp
--- a/Engine/Plugins/Media/WebMMedia/Source/WebMMedia/Private/WebMContainer.cpp
+++ b/Engine/Plugins/Media/WebMMedia/Source/WebMMedia/Private/WebMContainer.cpp
@@ -159,6 +159,7 @@ void FWebMContainer::ReadFrames(FTimespan ReadBufferLength, TArray<TSharedPtr<FW
 
 			TSharedPtr<FWebMFrame> Frame = MakeShared<FWebMFrame>();
 			Frame->Time = CurrentTime;
+			Frame->Duration = FMkvFileReader::GetVideoFrameDuration(*MkvFile->VideoTracks[SelectedVideoTrack]);
 			Frame->Data.SetNumUninitialized(MkvFrame.len);
 			MkvFrame.Read(MkvReader.Get(), Frame->Data.GetData());
 
diff --git a/Engine/Plugins/Media/WebMMedia/Source/WebMMedia/Public/WebMAudioDecoder.h b/Engine/Plugins/Media/WebMMedia/Source/WebMMedia/Public/WebMAudioDecoder.h
--- a/Engine/Plugins/Media/WebMMedia/Source/WebMMedia/Public/WebMAudioDecoder.h
+++ b/Engine/Plugins/Media/WebMMedia/Source/WebMMedia/Public/WebMAudioDecoder.h
@@ -6,7 +6,7 @@
 #include "Templates/SharedPointer.h"
 #include "Templates/UniquePtr.h"
 
-class FMediaSamples;
+class IWebMSamplesSink;
 class FWebMMediaAudioSample;
 class FWebMMediaAudioSamplePool;
 struct FWebMFrame;
@@ -15,7 +15,7 @@ struct OpusDecoder;
 class WEBMMEDIA_API FWebMAudioDecoder
 {
 public:
-	FWebMAudioDecoder(TSharedPtr<FMediaSamples, ESPMode::ThreadSafe> InSamples);
+	FWebMAudioDecoder(IWebMSamplesSink& InSamples);
 	~FWebMAudioDecoder();
 
 public:
@@ -31,11 +31,11 @@ private:
 	};
 
 	struct FVorbisDecoder;
-	TSharedPtr<FMediaSamples, ESPMode::ThreadSafe> Samples;
 	TUniquePtr<FWebMMediaAudioSamplePool> AudioSamplePool;
 	TUniquePtr<FVorbisDecoder> VorbisDecoder;
 	FGraphEventRef AudioDecodingTask;
 	TArray<uint8> DecodeBuffer;
+	IWebMSamplesSink& Samples;
 	OpusDecoder* OpusDecoder;
 	ESupportedCodecs Codec;
 	int32 FrameSize;
diff --git a/Engine/Plugins/Media/WebMMedia/Source/WebMMedia/Public/WebMMediaFrame.h b/Engine/Plugins/Media/WebMMedia/Source/WebMMedia/Public/WebMMediaFrame.h
--- a/Engine/Plugins/Media/WebMMedia/Source/WebMMedia/Public/WebMMediaFrame.h
+++ b/Engine/Plugins/Media/WebMMedia/Source/WebMMedia/Public/WebMMediaFrame.h
@@ -8,4 +8,5 @@ struct FWebMFrame
 {
 	TArray<uint8> Data;
 	FTimespan Time;
+	FTimespan Duration;
 };
diff --git a/Engine/Plugins/Media/WebMMedia/Source/WebMMedia/Public/WebMMediaTextureSample.h b/Engine/Plugins/Media/WebMMedia/Source/WebMMedia/Public/WebMMediaTextureSample.h
--- a/Engine/Plugins/Media/WebMMedia/Source/WebMMedia/Public/WebMMediaTextureSample.h
+++ b/Engine/Plugins/Media/WebMMedia/Source/WebMMedia/Public/WebMMediaTextureSample.h
@@ -18,7 +18,7 @@ public:
 	FWebMMediaTextureSample();
 
 public:
-	void Initialize(FIntPoint InDisplaySize, FIntPoint InTotalSize, FTimespan InTime);
+	void Initialize(FIntPoint InDisplaySize, FIntPoint InTotalSize, FTimespan InTime, FTimespan InDuration);
 	void CreateTexture();
 
 public:
@@ -44,6 +44,7 @@ public:
 private:
 	TRefCountPtr<FRHITexture2D> Texture;
 	FTimespan Time;
+	FTimespan Duration;
 	FIntPoint TotalSize;
 	FIntPoint DisplaySize;
 };
diff --git a/Engine/Plugins/Media/WebMMedia/Source/WebMMedia/Public/WebMSamplesSink.h b/Engine/Plugins/Media/WebMMedia/Source/WebMMedia/Public/WebMSamplesSink.h
new file mode 100644
--- /dev/null
+++ b/Engine/Plugins/Media/WebMMedia/Source/WebMMedia/Public/WebMSamplesSink.h
@@ -0,0 +1,18 @@
+// Copyright 1998-2019 Epic Games, Inc. All Rights Reserved.
+
+#pragma once
+
+#include "CoreMinimal.h"
+#include "Templates/SharedPointer.h"
+
+class FWebMMediaAudioSample;
+class FWebMMediaTextureSample;
+
+class IWebMSamplesSink
+{
+public:
+	virtual void AddVideoSampleFromDecodingThread(TSharedRef<FWebMMediaTextureSample, ESPMode::ThreadSafe> Sample) = 0;
+	virtual void AddAudioSampleFromDecodingThread(TSharedRef<FWebMMediaAudioSample, ESPMode::ThreadSafe> Sample) = 0;
+
+	virtual ~IWebMSamplesSink() {}
+};
diff --git a/Engine/Plugins/Media/WebMMedia/Source/WebMMedia/Public/WebMVideoDecoder.h b/Engine/Plugins/Media/WebMMedia/Source/WebMMedia/Public/WebMVideoDecoder.h
--- a/Engine/Plugins/Media/WebMMedia/Source/WebMMedia/Public/WebMVideoDecoder.h
+++ b/Engine/Plugins/Media/WebMMedia/Source/WebMMedia/Public/WebMVideoDecoder.h
@@ -12,7 +12,7 @@ THIRD_PARTY_INCLUDES_START
 #include <vpx/vp8dx.h>
 THIRD_PARTY_INCLUDES_END
 
-class FMediaSamples;
+class IWebMSamplesSink;
 class FWebMMediaTextureSample;
 class FWebMMediaTextureSamplePool;
 struct FWebMFrame;
@@ -20,7 +20,7 @@ struct FWebMFrame;
 class WEBMMEDIA_API FWebMVideoDecoder
 {
 public:
-	FWebMVideoDecoder(TSharedPtr<FMediaSamples, ESPMode::ThreadSafe> InSamples);
+	FWebMVideoDecoder(IWebMSamplesSink& InSamples);
 	~FWebMVideoDecoder();
 
 public:
@@ -36,12 +36,12 @@ private:
 	};
 
 	vpx_codec_ctx_t Context;
-	TSharedPtr<FMediaSamples, ESPMode::ThreadSafe> Samples;
 	TUniquePtr<FWebMMediaTextureSamplePool> VideoSamplePool;
 	TRefCountPtr<FRHITexture2D> DecodedY;
 	TRefCountPtr<FRHITexture2D> DecodedU;
 	TRefCountPtr<FRHITexture2D> DecodedV;
 	FGraphEventRef VideoDecodingTask;
+	IWebMSamplesSink& Samples;
 	bool bTexturesCreated;
 	bool bIsInitialized;
 
diff --git a/Engine/Plugins/Runtime/WebMMoviePlayer/Source/WebMMoviePlayer/Private/WebMMovieStreamer.cpp b/Engine/Plugins/Runtime/WebMMoviePlayer/Source/WebMMoviePlayer/Private/WebMMovieStreamer.cpp
--- a/Engine/Plugins/Runtime/WebMMoviePlayer/Source/WebMMoviePlayer/Private/WebMMovieStreamer.cpp
+++ b/Engine/Plugins/Runtime/WebMMoviePlayer/Source/WebMMoviePlayer/Private/WebMMovieStreamer.cpp
@@ -95,8 +95,8 @@ bool FWebMMovieStreamer::StartNextMovie()
 		}
 
 		Samples = MakeShareable(new FMediaSamples());
-		AudioDecoder.Reset(new FWebMAudioDecoder(Samples));
-		VideoDecoder.Reset(new FWebMVideoDecoder(Samples));
+		AudioDecoder.Reset(new FWebMAudioDecoder(*this));
+		VideoDecoder.Reset(new FWebMVideoDecoder(*this));
 
 		FWebMAudioTrackInfo DefaultAudioTrack = Container->GetCurrentAudioTrackInfo();
 		check(DefaultAudioTrack.bIsValid);
@@ -283,4 +283,14 @@ bool FWebMMovieStreamer::ReadMoreFrames()
 	return VideoFrames.Num() > 0 || AudioFrames.Num() > 0;
 }
 
+void FWebMMovieStreamer::AddVideoSampleFromDecodingThread(TSharedRef<FWebMMediaTextureSample, ESPMode::ThreadSafe> Sample)
+{
+	Samples->AddVideo(Sample);
+}
+
+void FWebMMovieStreamer::AddAudioSampleFromDecodingThread(TSharedRef<FWebMMediaAudioSample, ESPMode::ThreadSafe> Sample)
+{
+	Samples->AddAudio(Sample);
+}
+
 #endif // WITH_WEBM_LIBS
diff --git a/Engine/Plugins/Runtime/WebMMoviePlayer/Source/WebMMoviePlayer/Private/WebMMovieStreamer.h b/Engine/Plugins/Runtime/WebMMoviePlayer/Source/WebMMoviePlayer/Private/WebMMovieStreamer.h
--- a/Engine/Plugins/Runtime/WebMMoviePlayer/Source/WebMMoviePlayer/Private/WebMMovieStreamer.h
+++ b/Engine/Plugins/Runtime/WebMMoviePlayer/Source/WebMMoviePlayer/Private/WebMMovieStreamer.h
@@ -8,18 +8,23 @@
 #include "WebMAudioBackend.h"
 #include "WebMMediaFrame.h"
 #include "Containers/Queue.h"
+#include "WebMSamplesSink.h"
 
 class FWebMVideoDecoder;
 class FWebMAudioDecoder;
 class FMediaSamples;
 class FWebMContainer;
 
-class FWebMMovieStreamer : public IMovieStreamer
+class FWebMMovieStreamer : public IMovieStreamer, public IWebMSamplesSink
 {
 public:
 	FWebMMovieStreamer();
 	~FWebMMovieStreamer();
 
+	//~ IWebMSamplesSink interface
+	virtual void AddVideoSampleFromDecodingThread(TSharedRef<FWebMMediaTextureSample, ESPMode::ThreadSafe> Sample) override;
+	virtual void AddAudioSampleFromDecodingThread(TSharedRef<FWebMMediaAudioSample, ESPMode::ThreadSafe> Sample) override;
+
 	//~ IMediaOutput interface
 	virtual bool Init(const TArray<FString>& InMoviePaths, TEnumAsByte<EMoviePlaybackType> InPlaybackType) override;
 	virtual void ForceCompletion() override;
