WebM: don't allow decoding more than 20 frames at once

diff --git a/Engine/Plugins/Media/WebMMedia/Source/WebMMedia/Private/Player/WebMMediaPlayer.cpp b/Engine/Plugins/Media/WebMMedia/Source/WebMMedia/Private/Player/WebMMediaPlayer.cpp
--- a/Engine/Plugins/Media/WebMMedia/Source/WebMMedia/Private/Player/WebMMediaPlayer.cpp
+++ b/Engine/Plugins/Media/WebMMedia/Source/WebMMedia/Private/Player/WebMMediaPlayer.cpp
@@ -263,6 +263,12 @@ void FWebMMediaPlayer::TickFetch(FTimespan DeltaTime, FTimespan Timecode)
 				AudioFrames.Add(Frame);
 			}
 		}
+
+		// Make sure we're not decoding too much at once as it is GPU memory intensive.
+		if (VideoFrames.Num() > 20)
+		{
+			break;
+		}
 	}
 
 	// Trigger video decoding
