WebM: moved files from Private to Public to allow reusability (UE-65049).

(Edigrating 4465055 from Dev-Console to Release-4.21)

diff --git a/Engine/Plugins/Media/WebMMedia/Source/WebMMedia/Private/Player/WebMAudioDecoder.cpp b/Engine/Plugins/Media/WebMMedia/Source/WebMMedia/Private/Player/WebMAudioDecoder.cpp
--- a/Engine/Plugins/Media/WebMMedia/Source/WebMMedia/Private/Player/WebMAudioDecoder.cpp
+++ b/Engine/Plugins/Media/WebMMedia/Source/WebMMedia/Private/Player/WebMAudioDecoder.cpp
@@ -3,6 +3,7 @@
 #include "WebMAudioDecoder.h"
 #include "WebMMediaPrivate.h"
 #include "WebMMediaFrame.h"
+#include "WebMMediaAudioSample.h"
 #include "MediaSamples.h"
 
 THIRD_PARTY_INCLUDES_START
diff --git a/Engine/Plugins/Media/WebMMedia/Source/WebMMedia/Private/Player/WebMAudioDecoder.h b/Engine/Plugins/Media/WebMMedia/Source/WebMMedia/Private/Player/WebMAudioDecoder.h
deleted file mode 100644
diff --git a/Engine/Plugins/Media/WebMMedia/Source/WebMMedia/Private/Player/WebMMediaAudioSample.h b/Engine/Plugins/Media/WebMMedia/Source/WebMMedia/Private/Player/WebMMediaAudioSample.h
deleted file mode 100644
diff --git a/Engine/Plugins/Media/WebMMedia/Source/WebMMedia/Private/Player/WebMMediaFrame.h b/Engine/Plugins/Media/WebMMedia/Source/WebMMedia/Private/Player/WebMMediaFrame.h
deleted file mode 100644
diff --git a/Engine/Plugins/Media/WebMMedia/Source/WebMMedia/Private/Player/WebMMediaTextureSample.cpp b/Engine/Plugins/Media/WebMMedia/Source/WebMMedia/Private/Player/WebMMediaTextureSample.cpp
--- a/Engine/Plugins/Media/WebMMedia/Source/WebMMedia/Private/Player/WebMMediaTextureSample.cpp
+++ b/Engine/Plugins/Media/WebMMedia/Source/WebMMedia/Private/Player/WebMMediaTextureSample.cpp
@@ -97,3 +97,8 @@ void FWebMMediaTextureSample::ShutdownPoolable()
 {
 	Time = FTimespan::Zero();
 }
+
+TRefCountPtr<FRHITexture2D> FWebMMediaTextureSample::GetTextureRef() const
+{
+	return Texture;
+}
diff --git a/Engine/Plugins/Media/WebMMedia/Source/WebMMedia/Private/Player/WebMMediaTextureSample.h b/Engine/Plugins/Media/WebMMedia/Source/WebMMedia/Private/Player/WebMMediaTextureSample.h
deleted file mode 100644
diff --git a/Engine/Plugins/Media/WebMMedia/Source/WebMMedia/Private/Player/WebMVideoDecoder.cpp b/Engine/Plugins/Media/WebMMedia/Source/WebMMedia/Private/Player/WebMVideoDecoder.cpp
--- a/Engine/Plugins/Media/WebMMedia/Source/WebMMedia/Private/Player/WebMVideoDecoder.cpp
+++ b/Engine/Plugins/Media/WebMMedia/Source/WebMMedia/Private/Player/WebMVideoDecoder.cpp
@@ -3,11 +3,12 @@
 #include "WebMVideoDecoder.h"
 #include "WebMMediaPrivate.h"
 #include "WebMMediaFrame.h"
+#include "WebMMediaTextureSample.h"
 #include "MediaShaders.h"
 #include "MediaSamples.h"
 #include "PipelineStateCache.h"
 #include "RHIStaticStates.h"
-# include "Containers/DynamicRHIResourceArray.h"
+#include "Containers/DynamicRHIResourceArray.h"
 
 namespace
 {
diff --git a/Engine/Plugins/Media/WebMMedia/Source/WebMMedia/Private/Player/WebMVideoDecoder.h b/Engine/Plugins/Media/WebMMedia/Source/WebMMedia/Private/Player/WebMVideoDecoder.h
deleted file mode 100644
diff --git a/Engine/Plugins/Media/WebMMedia/Source/WebMMedia/Public/WebMAudioDecoder.h b/Engine/Plugins/Media/WebMMedia/Source/WebMMedia/Public/WebMAudioDecoder.h
new file mode 100644
--- /dev/null
+++ b/Engine/Plugins/Media/WebMMedia/Source/WebMMedia/Public/WebMAudioDecoder.h
@@ -0,0 +1,36 @@
+// Copyright 1998-2018 Epic Games, Inc. All Rights Reserved.
+
+#pragma once
+
+#include "Async/TaskGraphInterfaces.h"
+#include "Templates/SharedPointer.h"
+#include "Templates/UniquePtr.h"
+
+class FMediaSamples;
+class FWebMMediaAudioSample;
+class FWebMMediaAudioSamplePool;
+struct FWebMFrame;
+struct OpusDecoder;
+
+class WEBMMEDIA_API FWebMAudioDecoder
+{
+public:
+	FWebMAudioDecoder(TSharedPtr<FMediaSamples, ESPMode::ThreadSafe> InSamples);
+	~FWebMAudioDecoder();
+
+public:
+	void Initialize(int32 InSampleRate, int32 InChannels);
+	void DecodeAudioFramesAsync(const TArray<TSharedPtr<FWebMFrame>>& AudioFrames);
+
+private:
+	TSharedPtr<FMediaSamples, ESPMode::ThreadSafe> Samples;
+	TUniquePtr<FWebMMediaAudioSamplePool> AudioSamplePool;
+	FGraphEventRef AudioDecodingTask;
+	TArray<uint8> DecodeBuffer;
+	OpusDecoder* Decoder;
+	int32 FrameSize;
+	int32 SampleRate;
+	int32 Channels;
+
+	void DoDecodeAudioFrames(const TArray<TSharedPtr<FWebMFrame>>& AudioFrames);
+};
diff --git a/Engine/Plugins/Media/WebMMedia/Source/WebMMedia/Public/WebMMediaAudioSample.h b/Engine/Plugins/Media/WebMMedia/Source/WebMMedia/Public/WebMMediaAudioSample.h
new file mode 100644
--- /dev/null
+++ b/Engine/Plugins/Media/WebMMedia/Source/WebMMedia/Public/WebMMediaAudioSample.h
@@ -0,0 +1,122 @@
+// Copyright 1998-2018 Epic Games, Inc. All Rights Reserved.
+
+#pragma once
+
+#include "CoreTypes.h"
+#include "Containers/Array.h"
+#include "IMediaAudioSample.h"
+#include "MediaObjectPool.h"
+#include "MediaSampleQueue.h"
+#include "Math/IntPoint.h"
+#include "Misc/Timespan.h"
+
+
+/**
+ * Implements a media audio sample for WebMMedia.
+ */
+class FWebMMediaAudioSample
+	: public IMediaAudioSample
+	, public IMediaPoolable
+{
+public:
+
+	/** Default constructor. */
+	FWebMMediaAudioSample()
+		: Channels(0)
+		, Duration(FTimespan::Zero())
+		, SampleRate(0)
+		, Time(FTimespan::Zero())
+	{ }
+
+	/** Virtual destructor. */
+	virtual ~FWebMMediaAudioSample() { }
+
+public:
+
+	/**
+	 * Initialize the sample.
+	 *
+	 * @param InBuffer The sample's data buffer.
+	 * @param InSize The size of the sample buffer (in bytes).
+	 * @param InTime The sample time (relative to presentation clock).
+	 * @param InDuration The duration for which the sample is valid.
+	 */
+	void Initialize(
+		const uint8* InBuffer,
+		uint32 InSize,
+		uint32 InChannels,
+		uint32 InSampleRate,
+		FTimespan InTime,
+		FTimespan InDuration)
+	{
+		check(InBuffer && InSize > 0);
+
+		Buffer.Reset(InSize);
+		Buffer.Append(InBuffer, InSize);
+
+		Channels = InChannels;
+		Duration = InDuration;
+		SampleRate = InSampleRate;
+		Time = InTime;
+	}
+
+public:
+
+	//~ IMediaAudioSample interface
+
+	virtual const void* GetBuffer() override
+	{
+		return Buffer.GetData();
+	}
+
+	virtual uint32 GetChannels() const override
+	{
+		return Channels;
+	}
+
+	virtual FTimespan GetDuration() const override
+	{
+		return Duration;
+	}
+
+	virtual EMediaAudioSampleFormat GetFormat() const override
+	{
+		return EMediaAudioSampleFormat::Int16;
+	}
+
+	virtual uint32 GetFrames() const override
+	{
+		return Buffer.Num() / (Channels * sizeof(int16));
+	}
+
+	virtual uint32 GetSampleRate() const override
+	{
+		return SampleRate;
+	}
+
+	virtual FTimespan GetTime() const override
+	{
+		return Time;
+	}
+
+private:
+
+	/** The sample's data buffer. */
+	TArray<uint8> Buffer;
+
+	/** Number of audio channels. */
+	uint32 Channels;
+
+	/** The duration for which the sample is valid. */
+	FTimespan Duration;
+
+	/** Audio sample rate (in samples per second). */
+	uint32 SampleRate;
+
+	/** Presentation time for which the sample was generated. */
+	FTimespan Time;
+};
+
+
+/** Implements a pool for WebM audio sample objects. */
+class FWebMMediaAudioSamplePool : public TMediaObjectPool<FWebMMediaAudioSample> { };
diff --git a/Engine/Plugins/Media/WebMMedia/Source/WebMMedia/Public/WebMMediaFrame.h b/Engine/Plugins/Media/WebMMedia/Source/WebMMedia/Public/WebMMediaFrame.h
new file mode 100644
--- /dev/null
+++ b/Engine/Plugins/Media/WebMMedia/Source/WebMMedia/Public/WebMMediaFrame.h
@@ -0,0 +1,11 @@
+// Copyright 1998-2018 Epic Games, Inc. All Rights Reserved.
+
+#pragma once
+
+#include "CoreMinimal.h"
+
+struct FWebMFrame
+{
+	TArray<uint8> Data;
+	FTimespan Time;
+};
diff --git a/Engine/Plugins/Media/WebMMedia/Source/WebMMedia/Public/WebMMediaTextureSample.h b/Engine/Plugins/Media/WebMMedia/Source/WebMMedia/Public/WebMMediaTextureSample.h
new file mode 100644
--- /dev/null
+++ b/Engine/Plugins/Media/WebMMedia/Source/WebMMedia/Public/WebMMediaTextureSample.h
@@ -0,0 +1,51 @@
+// Copyright 1998-2018 Epic Games, Inc. All Rights Reserved.
+
+#pragma once
+
+#include "CoreMinimal.h"
+#include "IMediaTextureSample.h"
+#include "MediaObjectPool.h"
+#include "RHI.h"
+#include "RHIUtilities.h"
+
+class FRHITexture2D;
+
+class WEBMMEDIA_API FWebMMediaTextureSample
+	: public IMediaTextureSample
+	, public IMediaPoolable
+{
+public:
+	FWebMMediaTextureSample();
+
+public:
+	void Initialize(FIntPoint InDisplaySize, FIntPoint InTotalSize, FTimespan InTime);
+	void CreateTexture();
+
+public:
+	//~ IMediaTextureSample interface
+	virtual const void* GetBuffer() override;
+	virtual FIntPoint GetDim() const override;
+	virtual FTimespan GetDuration() const override;
+	virtual EMediaTextureSampleFormat GetFormat() const override;
+	virtual FIntPoint GetOutputDim() const override;
+	virtual uint32 GetStride() const override;
+	virtual FRHITexture* GetTexture() const override;
+	virtual FTimespan GetTime() const override;
+	virtual bool IsCacheable() const override;
+	virtual bool IsOutputSrgb() const override;
+
+public:
+	//~ IMediaPoolable interface
+	virtual void ShutdownPoolable() override;
+
+public:
+	TRefCountPtr<FRHITexture2D> GetTextureRef() const;
+
+private:
+	TRefCountPtr<FRHITexture2D> Texture;
+	FTimespan Time;
+	FIntPoint TotalSize;
+	FIntPoint DisplaySize;
+};
+
+class FWebMMediaTextureSamplePool : public TMediaObjectPool<FWebMMediaTextureSample> { };
diff --git a/Engine/Plugins/Media/WebMMedia/Source/WebMMedia/Public/WebMVideoDecoder.h b/Engine/Plugins/Media/WebMMedia/Source/WebMMedia/Public/WebMVideoDecoder.h
new file mode 100644
--- /dev/null
+++ b/Engine/Plugins/Media/WebMMedia/Source/WebMMedia/Public/WebMVideoDecoder.h
@@ -0,0 +1,48 @@
+// Copyright 1998-2018 Epic Games, Inc. All Rights Reserved.
+
+#pragma once
+
+#include "Templates/SharedPointer.h"
+#include "MediaShaders.h"
+
+THIRD_PARTY_INCLUDES_START
+#include <vpx/vpx_decoder.h>
+#include <vpx/vp8dx.h>
+THIRD_PARTY_INCLUDES_END
+
+class FMediaSamples;
+class FWebMMediaTextureSample;
+class FWebMMediaTextureSamplePool;
+struct FWebMFrame;
+
+class WEBMMEDIA_API FWebMVideoDecoder
+{
+public:
+	FWebMVideoDecoder(TSharedPtr<FMediaSamples, ESPMode::ThreadSafe> InSamples);
+	~FWebMVideoDecoder();
+
+public:
+	void Initialize();
+	void DecodeVideoFramesAsync(const TArray<TSharedPtr<FWebMFrame>>& VideoFrames);
+
+private:
+	struct FConvertParams
+	{
+		TSharedPtr<FWebMMediaTextureSample, ESPMode::ThreadSafe> VideoSample;
+		const vpx_image_t* Image;
+	};
+
+	TSharedPtr<FMediaSamples, ESPMode::ThreadSafe> Samples;
+	TUniquePtr<FWebMMediaTextureSamplePool> VideoSamplePool;
+	TRefCountPtr<FRHITexture2D> DecodedY;
+	TRefCountPtr<FRHITexture2D> DecodedU;
+	TRefCountPtr<FRHITexture2D> DecodedV;
+	FGraphEventRef VideoDecodingTask;
+	bool bTexturesCreated;
+
+	vpx_codec_ctx_t Context;
+
+	void ConvertYUVToRGBAndSubmit(const FConvertParams& Params);
+	void DoDecodeVideoFrames(const TArray<TSharedPtr<FWebMFrame>>& VideoFrames);
+	void CreateTextures(const vpx_image_t* Image);
+};
