WebM: VP8 and Vorbis codecs added, extension change from mkv to webm.

(Edigrating 4465277 from Dev-Console to Release-4.21)

diff --git a/Engine/Plugins/Media/WebMMedia/Source/WebMMedia/Private/Player/WebMAudioDecoder.cpp b/Engine/Plugins/Media/WebMMedia/Source/WebMMedia/Private/Player/WebMAudioDecoder.cpp
--- a/Engine/Plugins/Media/WebMMedia/Source/WebMMedia/Private/Player/WebMAudioDecoder.cpp
+++ b/Engine/Plugins/Media/WebMMedia/Source/WebMMedia/Private/Player/WebMAudioDecoder.cpp
@@ -8,12 +8,20 @@
 
 THIRD_PARTY_INCLUDES_START
 #include <opus.h>
+#include <vorbis/codec.h>
 THIRD_PARTY_INCLUDES_END
 
+struct FWebMAudioDecoder::FVorbisDecoder
+{
+	vorbis_info Info;
+	vorbis_dsp_state DspState;
+	vorbis_block Block;
+};
+
 FWebMAudioDecoder::FWebMAudioDecoder(TSharedPtr<FMediaSamples, ESPMode::ThreadSafe> InSamples)
 	: Samples(InSamples)
 	, AudioSamplePool(new FWebMMediaAudioSamplePool)
-	, Decoder(nullptr)
+	, OpusDecoder(nullptr)
 {
 }
 
@@ -24,25 +32,153 @@ FWebMAudioDecoder::~FWebMAudioDecoder()
 		FTaskGraphInterface::Get().WaitUntilTaskCompletes(AudioDecodingTask);
 	}
 
-	if (Decoder)
+	if (OpusDecoder)
 	{
-		opus_decoder_destroy(Decoder);
+		opus_decoder_destroy(OpusDecoder);
 	}
 }
 
-void FWebMAudioDecoder::Initialize(int32 InSampleRate, int32 InChannels)
+bool FWebMAudioDecoder::Initialize(const char* InCodec, int32 InSampleRate, int32 InChannels, const uint8* CodecPrivateData, size_t CodecPrivateDataSize)
 {
+	if (OpusDecoder)
+	{
+		opus_decoder_destroy(OpusDecoder);
+		OpusDecoder = nullptr;
+	}
+
+	if (FCStringAnsi::Strcmp(InCodec, "A_OPUS") == 0)
+	{
+		Codec = ESupportedCodecs::Opus;
+	}
+	else if (FCStringAnsi::Strcmp(InCodec, "A_VORBIS") == 0)
+	{
+		Codec = ESupportedCodecs::Vorbis;
+	}
+	else
+	{
+		return false;
+	}
+
 	SampleRate = InSampleRate;
 	Channels = InChannels;
-	// Max supported frame is 120ms
-	FrameSize = 120 * SampleRate / 1000;
 
+	if (Codec == ESupportedCodecs::Opus)
+	{
+		InitializeOpus();
+	}
+	else
+	{
+		if (!InitializeVorbis(CodecPrivateData, CodecPrivateDataSize))
+		{
+			return false;
+		}
+	}
+
+	DecodeBuffer.SetNumUninitialized(FrameSize * Channels * sizeof(int16));
+
+	return true;
+}
+
+void FWebMAudioDecoder::InitializeOpus()
+{
 	int32 ErrorCode = 0;
-	Decoder = opus_decoder_create(SampleRate, Channels, &ErrorCode);
-	check(Decoder);
+	OpusDecoder = opus_decoder_create(SampleRate, Channels, &ErrorCode);
+	check(OpusDecoder);
 	check(ErrorCode == 0);
 
-	DecodeBuffer.SetNumUninitialized(FrameSize * Channels * sizeof(opus_int16));
+	// Max supported frame is 120ms
+	FrameSize = 120 * SampleRate / 1000;
+}
+
+bool FWebMAudioDecoder::InitializeVorbis(const uint8* CodecPrivateData, size_t CodecPrivateDataSize)
+{
+	if (CodecPrivateDataSize < 3 || !CodecPrivateData || CodecPrivateData[0] != 2)
+	{
+		UE_LOG(LogWebMMedia, Display, TEXT("Failed to initialize Vorbis - invalid data"));
+		return false;
+	}
+
+	size_t HeaderSize[3] = { 0 };
+	size_t Offset = 1;
+
+	// Calculate three headers sizes
+	for (int32 i = 0; i < 2; ++i)
+	{
+		for (;;)
+		{
+			if (Offset >= CodecPrivateDataSize)
+			{
+				UE_LOG(LogWebMMedia, Display, TEXT("Failed to initialize Vorbis - invalid offset"));
+				return false;
+			}
+			HeaderSize[i] += CodecPrivateData[Offset];
+			if (CodecPrivateData[Offset++] < 0xFF)
+			{
+				break;
+			}
+		}
+	}
+	HeaderSize[2] = CodecPrivateDataSize - (HeaderSize[0] + HeaderSize[1] + Offset);
+
+	if (HeaderSize[0] + HeaderSize[1] + HeaderSize[2] + Offset != CodecPrivateDataSize)
+	{
+		UE_LOG(LogWebMMedia, Display, TEXT("Failed to initialize Vorbis - parameters mismatch"));
+		return false;
+	}
+
+	ogg_packet Packet[3];
+	memset(Packet, 0, 3 * sizeof(ogg_packet));
+
+	Packet[0].packet = (unsigned char *) CodecPrivateData + Offset;
+	Packet[0].bytes = HeaderSize[0];
+	Packet[0].b_o_s = 1;
+
+	Packet[1].packet = (unsigned char *) CodecPrivateData + Offset + HeaderSize[0];
+	Packet[1].bytes = HeaderSize[1];
+
+	Packet[2].packet = (unsigned char *) CodecPrivateData + Offset + HeaderSize[0] + HeaderSize[1];
+	Packet[2].bytes = HeaderSize[2];
+
+	VorbisDecoder.Reset(new FVorbisDecoder());
+
+	vorbis_info_init(&VorbisDecoder->Info);
+
+	/* Upload three Vorbis headers into libvorbis */
+	vorbis_comment Vc;
+	vorbis_comment_init(&Vc);
+	for (int i = 0; i < 3; ++i)
+	{
+		if (vorbis_synthesis_headerin(&VorbisDecoder->Info, &Vc, &Packet[i]))
+		{
+			UE_LOG(LogWebMMedia, Display, TEXT("Failed to initialize Vorbis - invalid headers"));
+			vorbis_comment_clear(&Vc);
+			return false;
+		}
+	}
+
+	vorbis_comment_clear(&Vc);
+
+	if (vorbis_synthesis_init(&VorbisDecoder->DspState, &VorbisDecoder->Info))
+	{
+		UE_LOG(LogWebMMedia, Display, TEXT("Failed to initialize Vorbis (synthesis)"));
+		return false;
+	}
+
+	if (VorbisDecoder->Info.channels != Channels || VorbisDecoder->Info.rate != SampleRate)
+	{
+		UE_LOG(LogWebMMedia, Display, TEXT("Failed to initialize Vorbis - invalid parameters"));
+		return false;
+	}
+
+	if (vorbis_block_init(&VorbisDecoder->DspState, &VorbisDecoder->Block))
+	{
+		UE_LOG(LogWebMMedia, Display, TEXT("Failed to initialize Vorbis"));
+		return false;
+	}
+
+	FrameSize = 4096 / Channels;
+
+	return true;
 }
 
 void FWebMAudioDecoder::DecodeAudioFramesAsync(const TArray<TSharedPtr<FWebMFrame>>& AudioFrames)
@@ -64,16 +200,79 @@ void FWebMAudioDecoder::DoDecodeAudioFrames(const TArray<TSharedPtr<FWebMFrame>>
 {
 	for (const TSharedPtr<FWebMFrame>& AudioFrame : AudioFrames)
 	{
-		const TSharedRef<FWebMMediaAudioSample, ESPMode::ThreadSafe> AudioSample = AudioSamplePool->AcquireShared();
+		int32 NumOfSamplesDecoded;
+
+		if (Codec == ESupportedCodecs::Opus)
+		{
+			NumOfSamplesDecoded = DecodeOpus(AudioFrame);
+		}
+		else
+		{
+			NumOfSamplesDecoded = DecodeVorbis(AudioFrame);
+		}
+
+		if (NumOfSamplesDecoded > 0)
+		{
+			TSharedRef<FWebMMediaAudioSample, ESPMode::ThreadSafe> AudioSample = AudioSamplePool->AcquireShared();
+			AudioSample->Initialize(DecodeBuffer.GetData(), NumOfSamplesDecoded * Channels * 2, Channels, SampleRate, AudioFrame->Time, FTimespan::FromHours(1));
+			Samples->AddAudio(AudioSample);
+		}
+	}
+}
+
+int32 FWebMAudioDecoder::DecodeOpus(const TSharedPtr<FWebMFrame>& AudioFrame)
+{
+	int32 Result = opus_decode(OpusDecoder, AudioFrame->Data.GetData(), AudioFrame->Data.Num(), (opus_int16*) DecodeBuffer.GetData(), FrameSize, 0);
+	if (Result < 0)
+	{
+		UE_LOG(LogWebMMedia, Display, TEXT("Error decoding Opus audio frame"));
+	}
 
-		int32 NumOfSamplesDecoded = opus_decode(Decoder, AudioFrame->Data.GetData(), AudioFrame->Data.Num(), (opus_int16*) DecodeBuffer.GetData(), FrameSize, 0);
-		if (NumOfSamplesDecoded <= 0)
+	return Result;
+}
+
+int32 FWebMAudioDecoder::DecodeVorbis(const TSharedPtr<FWebMFrame>& AudioFrame)
+{
+	ogg_packet Packet;
+	memset(&Packet, 0, sizeof(ogg_packet));
+
+	Packet.packet = AudioFrame->Data.GetData();
+	Packet.bytes = AudioFrame->Data.Num();
+
+	if (vorbis_synthesis(&VorbisDecoder->Block, &Packet))
+	{
+		UE_LOG(LogWebMMedia, Display, TEXT("Error decoding Vorbis audio frame - vorbis_synthesis failed"));
+		return -1;
+	}
+
+	if (vorbis_synthesis_blockin(&VorbisDecoder->DspState, &VorbisDecoder->Block))
+	{
+		UE_LOG(LogWebMMedia, Display, TEXT("Error decoding Vorbis audio frame - vorbis_synthesis_blockin failed"));
+		return - 1;
+	}
+
+	int16* Buffer = (int16*) DecodeBuffer.GetData();
+	int32 Count = 0;
+	float** Pcm;
+	int32 SamplesCount = vorbis_synthesis_pcmout(&VorbisDecoder->DspState, &Pcm);
+	while (SamplesCount > 0)
+	{
+		int32 ToConvert = SamplesCount <= FrameSize ? SamplesCount : FrameSize;
+		for (int32 Channel = 0; Channel < Channels; ++Channel)
 		{
-			UE_LOG(LogWebMMedia, Display, TEXT("Error decoding audio frame"));
-			return;
+			float *FloatSamples = Pcm[Channel];
+			for (int32 i = 0, j = Channel; i < ToConvert; ++i, j += Channels)
+			{
+				int32 Sample = FMath::Clamp<int32>(FloatSamples[i] * 32767.0f, -32768, 32768);
+				Buffer[Count + j] = Sample;
+			}
 		}
+		vorbis_synthesis_read(&VorbisDecoder->DspState, ToConvert);
 
-		AudioSample->Initialize(DecodeBuffer.GetData(), NumOfSamplesDecoded * Channels * 2, Channels, SampleRate, AudioFrame->Time, FTimespan::FromHours(1));
-		Samples->AddAudio(AudioSample);
+		Count += ToConvert;
+
+		SamplesCount = vorbis_synthesis_pcmout(&VorbisDecoder->DspState, &Pcm);
 	}
+
+	return Count;
 }
diff --git a/Engine/Plugins/Media/WebMMedia/Source/WebMMedia/Private/Player/WebMMediaPlayer.cpp b/Engine/Plugins/Media/WebMMedia/Source/WebMMedia/Private/Player/WebMMediaPlayer.cpp
--- a/Engine/Plugins/Media/WebMMedia/Source/WebMMedia/Private/Player/WebMMediaPlayer.cpp
+++ b/Engine/Plugins/Media/WebMMedia/Source/WebMMedia/Private/Player/WebMMediaPlayer.cpp
@@ -139,7 +139,7 @@ bool FWebMMediaPlayer::Open(const FString& Url, const IMediaOptions* /*Options*/
 
 	if (!MkvRead())
 	{
-		UE_LOG(LogWebMMedia, Error, TEXT("Error parsing MKV file: %s"), *FilePath);
+		UE_LOG(LogWebMMedia, Error, TEXT("Error parsing matroska file: %s"), *FilePath);
 		return false;
 	}
 
@@ -148,12 +148,19 @@ bool FWebMMediaPlayer::Open(const FString& Url, const IMediaOptions* /*Options*/
 
 	if (SelectedAudioTrack != INDEX_NONE)
 	{
-		AudioDecoder->Initialize(AudioTracks[SelectedAudioTrack]->GetSamplingRate(), AudioTracks[SelectedAudioTrack]->GetChannels());
+		const mkvparser::AudioTrack* AudioTrack = AudioTracks[SelectedAudioTrack];
+
+		size_t PrivateDataSize;
+		const uint8* PrivateData = AudioTrack->GetCodecPrivate(PrivateDataSize);
+
+		AudioDecoder->Initialize(AudioTrack->GetCodecId(), AudioTrack->GetSamplingRate(), AudioTrack->GetChannels(), PrivateData, PrivateDataSize);
 	}
 
 	if (SelectedVideoTrack != INDEX_NONE)
 	{
-		VideoDecoder->Initialize();
+		const mkvparser::VideoTrack* VideoTrack = VideoTracks[SelectedVideoTrack];
+
+		VideoDecoder->Initialize(VideoTrack->GetCodecId());
 	}
 
 	CurrentState = EMediaState::Stopped;
@@ -413,12 +420,19 @@ bool FWebMMediaPlayer::SelectTrack(EMediaTrackType TrackType, int32 TrackIndex)
 		{
 			if (SelectedAudioTrack != INDEX_NONE)
 			{
-				AudioDecoder->Initialize(AudioTracks[SelectedAudioTrack]->GetSamplingRate(), AudioTracks[SelectedAudioTrack]->GetChannels());
+				const mkvparser::AudioTrack* AudioTrack = AudioTracks[SelectedAudioTrack];
+
+				size_t PrivateDataSize;
+				const uint8* PrivateData = AudioTrack->GetCodecPrivate(PrivateDataSize);
+
+				AudioDecoder->Initialize(AudioTrack->GetCodecId(), AudioTrack->GetSamplingRate(), AudioTrack->GetChannels(), PrivateData, PrivateDataSize);
 			}
 
 			if (SelectedVideoTrack != INDEX_NONE)
 			{
-				VideoDecoder->Initialize();
+				const mkvparser::VideoTrack* VideoTrack = VideoTracks[SelectedVideoTrack];
+
+				VideoDecoder->Initialize(VideoTrack->GetCodecId());
 			}
 		}
 	}
@@ -526,12 +540,19 @@ bool FWebMMediaPlayer::Seek(const FTimespan& Time)
 
 	if (SelectedVideoTrack != INDEX_NONE)
 	{
-		VideoDecoder->Initialize();
+		const mkvparser::VideoTrack* VideoTrack = VideoTracks[SelectedVideoTrack];
+
+		VideoDecoder->Initialize(VideoTrack->GetCodecId());
 	}
 
 	if (SelectedAudioTrack != INDEX_NONE)
 	{
-		AudioDecoder->Initialize(AudioTracks[SelectedAudioTrack]->GetSamplingRate(), AudioTracks[SelectedAudioTrack]->GetChannels());
+		const mkvparser::AudioTrack* AudioTrack = AudioTracks[SelectedAudioTrack];
+
+		size_t PrivateDataSize;
+		const uint8* PrivateData = AudioTrack->GetCodecPrivate(PrivateDataSize);
+
+		AudioDecoder->Initialize(AudioTrack->GetCodecId(), AudioTrack->GetSamplingRate(), AudioTrack->GetChannels(), PrivateData, PrivateDataSize);
 	}
 
 	CurrentTime = Time;
@@ -606,7 +627,7 @@ bool FWebMMediaPlayer::MkvRead()
 
 		if (Track->GetType() == mkvparser::Track::kVideo)
 		{
-			if (FCStringAnsi::Strcmp(Track->GetCodecId(), "V_VP9") == 0)
+			if (FCStringAnsi::Strcmp(Track->GetCodecId(), "V_VP8") == 0 || FCStringAnsi::Strcmp(Track->GetCodecId(), "V_VP9") == 0)
 			{
 				VideoTracks.Add(static_cast<const mkvparser::VideoTrack*>(Track));
 			}
@@ -618,7 +639,7 @@ bool FWebMMediaPlayer::MkvRead()
 		}
 		else if (Track->GetType() == mkvparser::Track::kAudio)
 		{
-			if (FCStringAnsi::Strcmp(Track->GetCodecId(), "A_OPUS") == 0)
+			if (FCStringAnsi::Strcmp(Track->GetCodecId(), "A_OPUS") == 0 || FCStringAnsi::Strcmp(Track->GetCodecId(), "A_VORBIS") == 0)
 			{
 				AudioTracks.Add(static_cast<const mkvparser::AudioTrack*>(Track));
 			}
diff --git a/Engine/Plugins/Media/WebMMedia/Source/WebMMedia/Private/Player/WebMVideoDecoder.cpp b/Engine/Plugins/Media/WebMMedia/Source/WebMMedia/Private/Player/WebMVideoDecoder.cpp
--- a/Engine/Plugins/Media/WebMMedia/Source/WebMMedia/Private/Player/WebMVideoDecoder.cpp
+++ b/Engine/Plugins/Media/WebMMedia/Source/WebMMedia/Private/Player/WebMVideoDecoder.cpp
@@ -62,31 +62,39 @@ FWebMVideoDecoder::FWebMVideoDecoder(TSharedPtr<FMediaSamples, ESPMode::ThreadSa
 	: Samples(InSamples)
 	, VideoSamplePool(new FWebMMediaTextureSamplePool)
 	, bTexturesCreated(false)
+	, bIsInitialized(false)
 {
-	const int32 NumOfThreads = 1;
-	const vpx_codec_dec_cfg_t CodecConfig = { NumOfThreads, 0, 0 };
-	verify(vpx_codec_dec_init(&Context, vpx_codec_vp9_dx(), &CodecConfig, /*VPX_CODEC_USE_FRAME_THREADING*/ 0) == 0);
 }
 
 FWebMVideoDecoder::~FWebMVideoDecoder()
 {
-	if (VideoDecodingTask && !VideoDecodingTask->IsComplete())
-	{
-		FTaskGraphInterface::Get().WaitUntilTaskCompletes(VideoDecodingTask);
-	}
-
-	// Make sure all compute shader decoding is done
-	FlushRenderingCommands();
-
-	vpx_codec_destroy(&Context);
+	Close();
 }
 
-void FWebMVideoDecoder::Initialize()
+bool FWebMVideoDecoder::Initialize(const char* CodecName)
 {
-	// Flush decoder
-	vpx_codec_decode(&Context, nullptr, 0, nullptr, 0);
+	Close();
 
-	bTexturesCreated = false;
+	const int32 NumOfThreads = 1;
+	const vpx_codec_dec_cfg_t CodecConfig = { NumOfThreads, 0, 0 };
+
+	if (FCStringAnsi::Strcmp(CodecName, "V_VP8") == 0)
+	{
+		verify(vpx_codec_dec_init(&Context, vpx_codec_vp8_dx(), &CodecConfig, /*VPX_CODEC_USE_FRAME_THREADING*/ 0) == 0);
+	}
+	else if (FCStringAnsi::Strcmp(CodecName, "V_VP9") == 0)
+	{
+		verify(vpx_codec_dec_init(&Context, vpx_codec_vp9_dx(), &CodecConfig, /*VPX_CODEC_USE_FRAME_THREADING*/ 0) == 0);
+	}
+	else
+	{
+		UE_LOG(LogWebMMedia, Display, TEXT("Unsupported video codec: %s"), CodecName);
+		return false;
+	}
+
+	bIsInitialized = true;
+
+	return true;
 }
 
 void FWebMVideoDecoder::DecodeVideoFramesAsync(const TArray<TSharedPtr<FWebMFrame>>& VideoFrames)
@@ -156,6 +164,25 @@ void FWebMVideoDecoder::CreateTextures(const vpx_image_t* Image)
 	DecodedV = RHICreateTexture2D(Image->stride[2], Image->d_h / 2, PF_G8, 1, 1, TexCreate_Dynamic, CreateInfo);
 }
 
+void FWebMVideoDecoder::Close()
+{
+	if (VideoDecodingTask && !VideoDecodingTask->IsComplete())
+	{
+		FTaskGraphInterface::Get().WaitUntilTaskCompletes(VideoDecodingTask);
+	}
+
+	// Make sure all compute shader decoding is done
+	FlushRenderingCommands();
+
+	if (bIsInitialized)
+	{
+		vpx_codec_destroy(&Context);
+		bIsInitialized = false;
+	}
+
+	bTexturesCreated = false;
+}
+
 void FWebMVideoDecoder::ConvertYUVToRGBAndSubmit(const FConvertParams& Params)
 {
 	TSharedPtr<FWebMMediaTextureSample, ESPMode::ThreadSafe> VideoSample = Params.VideoSample;
diff --git a/Engine/Plugins/Media/WebMMedia/Source/WebMMedia/Public/WebMAudioDecoder.h b/Engine/Plugins/Media/WebMMedia/Source/WebMMedia/Public/WebMAudioDecoder.h
--- a/Engine/Plugins/Media/WebMMedia/Source/WebMMedia/Public/WebMAudioDecoder.h
+++ b/Engine/Plugins/Media/WebMMedia/Source/WebMMedia/Public/WebMAudioDecoder.h
@@ -19,18 +19,31 @@ public:
 	~FWebMAudioDecoder();
 
 public:
-	void Initialize(int32 InSampleRate, int32 InChannels);
+	bool Initialize(const char* InCodec, int32 InSampleRate, int32 InChannels, const uint8* CodecPrivateData, size_t CodecPrivateDataSize);
 	void DecodeAudioFramesAsync(const TArray<TSharedPtr<FWebMFrame>>& AudioFrames);
 
 private:
+	enum ESupportedCodecs
+	{
+		Opus,
+		Vorbis,
+	};
+
+	struct FVorbisDecoder;
 	TSharedPtr<FMediaSamples, ESPMode::ThreadSafe> Samples;
 	TUniquePtr<FWebMMediaAudioSamplePool> AudioSamplePool;
+	TUniquePtr<FVorbisDecoder> VorbisDecoder;
 	FGraphEventRef AudioDecodingTask;
 	TArray<uint8> DecodeBuffer;
-	OpusDecoder* Decoder;
+	OpusDecoder* OpusDecoder;
+	ESupportedCodecs Codec;
 	int32 FrameSize;
 	int32 SampleRate;
 	int32 Channels;
 
+	void InitializeOpus();
+	bool InitializeVorbis(const uint8* CodecPrivateData, size_t CodecPrivateDataSize);
 	void DoDecodeAudioFrames(const TArray<TSharedPtr<FWebMFrame>>& AudioFrames);
+	int32 DecodeOpus(const TSharedPtr<FWebMFrame>& AudioFrame);
+	int32 DecodeVorbis(const TSharedPtr<FWebMFrame>& AudioFrame);
 };
diff --git a/Engine/Plugins/Media/WebMMedia/Source/WebMMedia/Public/WebMVideoDecoder.h b/Engine/Plugins/Media/WebMMedia/Source/WebMMedia/Public/WebMVideoDecoder.h
--- a/Engine/Plugins/Media/WebMMedia/Source/WebMMedia/Public/WebMVideoDecoder.h
+++ b/Engine/Plugins/Media/WebMMedia/Source/WebMMedia/Public/WebMVideoDecoder.h
@@ -22,7 +22,7 @@ public:
 	~FWebMVideoDecoder();
 
 public:
-	void Initialize();
+	bool Initialize(const char* CodecName);
 	void DecodeVideoFramesAsync(const TArray<TSharedPtr<FWebMFrame>>& VideoFrames);
 
 private:
@@ -32,6 +32,7 @@ private:
 		const vpx_image_t* Image;
 	};
 
+	vpx_codec_ctx_t Context;
 	TSharedPtr<FMediaSamples, ESPMode::ThreadSafe> Samples;
 	TUniquePtr<FWebMMediaTextureSamplePool> VideoSamplePool;
 	TRefCountPtr<FRHITexture2D> DecodedY;
@@ -39,10 +40,10 @@ private:
 	TRefCountPtr<FRHITexture2D> DecodedV;
 	FGraphEventRef VideoDecodingTask;
 	bool bTexturesCreated;
-
-	vpx_codec_ctx_t Context;
+	bool bIsInitialized;
 
 	void ConvertYUVToRGBAndSubmit(const FConvertParams& Params);
 	void DoDecodeVideoFrames(const TArray<TSharedPtr<FWebMFrame>>& VideoFrames);
 	void CreateTextures(const vpx_image_t* Image);
+	void Close();
 };
diff --git a/Engine/Plugins/Media/WebMMedia/Source/WebMMedia/WebMMedia.Build.cs b/Engine/Plugins/Media/WebMMedia/Source/WebMMedia/WebMMedia.Build.cs
--- a/Engine/Plugins/Media/WebMMedia/Source/WebMMedia/WebMMedia.Build.cs
+++ b/Engine/Plugins/Media/WebMMedia/Source/WebMMedia/WebMMedia.Build.cs
@@ -7,38 +7,23 @@ namespace UnrealBuildTool.Rules
 	{
 		public WebMMedia(ReadOnlyTargetRules Target) : base(Target)
 		{
-			DynamicallyLoadedModuleNames.AddRange(
-				new string[] {
-					"Media",
-				});
-
 			PrivateDependencyModuleNames.AddRange(
 				new string[] {
 					"WebMMediaFactory",
 					"Core",
 					"Engine",
+					"Media",
 					"MediaUtils",
 					"RenderCore",
 					"RHI",
 					"ShaderCore",
 					"UtilityShaders",
 					"libOpus",
-					//"UEOgg",
-					//"Vorbis",
+					"UEOgg",
+					"Vorbis",
 					"LibVpx",
 					"LibWebM",
 				});
-
-			PrivateIncludePathModuleNames.AddRange(
-				new string[] {
-					"Media",
-				});
-
-			PrivateIncludePaths.AddRange(
-				new string[] {
-					"WebMMedia/Private",
-					"WebMMedia/Private/Player",
-				});
 		}
 	}
 }
diff --git a/Engine/Plugins/Media/WebMMedia/Source/WebMMediaEditor/Private/WebMFileMediaSourceFactory.cpp b/Engine/Plugins/Media/WebMMedia/Source/WebMMediaEditor/Private/WebMFileMediaSourceFactory.cpp
--- a/Engine/Plugins/Media/WebMMedia/Source/WebMMediaEditor/Private/WebMFileMediaSourceFactory.cpp
+++ b/Engine/Plugins/Media/WebMMedia/Source/WebMMediaEditor/Private/WebMFileMediaSourceFactory.cpp
@@ -13,7 +13,7 @@ UWebMPlatFileMediaSourceFactory::UWebMPlatFileMediaSourceFactory(const FObjectIn
 	: Super(ObjectInitializer)
 {
 	// add formats we support
-	Formats.Add(TEXT("mkv;Matroska Multimedia File"));
+	Formats.Add(TEXT("webm;WebM Multimedia File"));
 
 	SupportedClass = UFileMediaSource::StaticClass();
 
diff --git a/Engine/Plugins/Media/WebMMedia/Source/WebMMediaFactory/Private/WebMMediaFactoryModule.cpp b/Engine/Plugins/Media/WebMMedia/Source/WebMMediaFactory/Private/WebMMediaFactoryModule.cpp
--- a/Engine/Plugins/Media/WebMMedia/Source/WebMMediaFactory/Private/WebMMediaFactoryModule.cpp
+++ b/Engine/Plugins/Media/WebMMedia/Source/WebMMediaFactory/Private/WebMMediaFactoryModule.cpp
@@ -111,7 +111,7 @@ public:
 	virtual void StartupModule() override
 	{
 		// supported file extensions
-		SupportedFileExtensions.Add(TEXT("mkv"));
+		SupportedFileExtensions.Add(TEXT("webm"));
 
 		// if the module got built and run succesfully on this platform, it means it's supported
 		SupportedPlatforms.Add(FPlatformMisc::GetUBTPlatform());
