WebM: added media streamer implementation (UE-65049)

(Edigrating 4471246 from Dev-Console to Release-4.21)

diff --git a/Engine/Plugins/Media/WebMMedia/Source/WebMMedia/Private/Player/WebMVideoDecoder.cpp b/Engine/Plugins/Media/WebMMedia/Source/WebMMedia/Private/Player/WebMVideoDecoder.cpp
--- a/Engine/Plugins/Media/WebMMedia/Source/WebMMedia/Private/Player/WebMVideoDecoder.cpp
+++ b/Engine/Plugins/Media/WebMMedia/Source/WebMMedia/Private/Player/WebMVideoDecoder.cpp
@@ -112,6 +112,11 @@ void FWebMVideoDecoder::DecodeVideoFramesAsync(const TArray<TSharedPtr<FWebMFram
 	}, TStatId(), nullptr, ENamedThreads::AnyThread);
 }
 
+bool FWebMVideoDecoder::IsBusy() const
+{
+	return VideoDecodingTask && !VideoDecodingTask->IsComplete();
+}
+
 void FWebMVideoDecoder::DoDecodeVideoFrames(const TArray<TSharedPtr<FWebMFrame>>& VideoFrames)
 {
 	for (const TSharedPtr<FWebMFrame>& VideoFrame : VideoFrames)
diff --git a/Engine/Plugins/Media/WebMMedia/Source/WebMMedia/Private/WebMContainer.cpp b/Engine/Plugins/Media/WebMMedia/Source/WebMMedia/Private/WebMContainer.cpp
new file mode 100644
--- /dev/null
+++ b/Engine/Plugins/Media/WebMMedia/Source/WebMMedia/Private/WebMContainer.cpp
@@ -0,0 +1,263 @@
+// Copyright 1998-2018 Epic Games, Inc. All Rights Reserved.
+
+#include "WebMContainer.h"
+#include "Player/MkvFileReader.h"
+#include "WebMMediaPrivate.h"
+#include "WebMMediaFrame.h"
+
+struct FWebMContainer::FMkvFileState
+{
+	TArray<const mkvparser::VideoTrack*> VideoTracks;
+	TArray<const mkvparser::AudioTrack*> AudioTracks;
+	TUniquePtr<mkvparser::Segment> CurrentSegment;
+	const mkvparser::Cluster* CurrentCluster;
+	const mkvparser::BlockEntry* CurrentBlockEntry;
+};
+
+FWebMContainer::FWebMContainer()
+	: MkvFile(new FMkvFileState())
+	, SelectedAudioTrack(INDEX_NONE)
+	, SelectedVideoTrack(INDEX_NONE)
+	, bNoMoreToRead(false)
+{
+	MkvFile->CurrentSegment = nullptr;
+	MkvFile->CurrentBlockEntry = nullptr;
+}
+
+FWebMContainer::~FWebMContainer()
+{
+}
+
+bool FWebMContainer::Open(const FString& FilePath)
+{
+	MkvReader.Reset(new FMkvFileReader());
+	if (!MkvReader->Open(*FilePath))
+	{
+		UE_LOG(LogWebMMedia, Error, TEXT("Failed opening video file: %s"), *FilePath);
+		return false;
+	}
+
+	int64 FilePosition = 0;
+	if (mkvparser::EBMLHeader().Parse(MkvReader.Get(), FilePosition) != 0)
+	{
+		UE_LOG(LogWebMMedia, Error, TEXT("Failed parsing MKV header"));
+		return false;
+	}
+
+	mkvparser::Segment* Segment;
+	if (mkvparser::Segment::CreateInstance(MkvReader.Get(), FilePosition, Segment) != 0)
+	{
+		UE_LOG(LogWebMMedia, Error, TEXT("Failed parsing MKV"));
+		return false;
+	}
+
+	MkvFile->CurrentSegment.Reset(Segment);
+
+	if (MkvFile->CurrentSegment->Load() < 0)
+	{
+		UE_LOG(LogWebMMedia, Error, TEXT("Failed loading MKV header"));
+		return false;
+	}
+
+	// Read all the tracks
+
+	const mkvparser::Tracks* Tracks = MkvFile->CurrentSegment->GetTracks();
+	int32 NumOfTracks = Tracks->GetTracksCount();
+
+	if (NumOfTracks == 0)
+	{
+		UE_LOG(LogWebMMedia, Warning, TEXT("File doesn't have any tracks"));
+		return false;
+	}
+
+	for (int32 i = 0; i < NumOfTracks; ++i)
+	{
+		const mkvparser::Track* Track = Tracks->GetTrackByIndex(i);
+		check(Track);
+
+		if (Track->GetType() == mkvparser::Track::kVideo)
+		{
+			if (FCStringAnsi::Strcmp(Track->GetCodecId(), "V_VP9") == 0)
+			{
+				MkvFile->VideoTracks.Add(static_cast<const mkvparser::VideoTrack*>(Track));
+
+				// Set as default
+				if (SelectedVideoTrack == INDEX_NONE)
+				{
+					SelectedVideoTrack = MkvFile->VideoTracks.Num() - 1;
+				}
+				
+			}
+			else
+			{
+				UE_LOG(LogWebMMedia, Warning, TEXT("File contains unsupported video track %d: %s"), i, UTF8_TO_TCHAR(Track->GetCodecId()));
+				continue;
+			}
+		}
+		else if (Track->GetType() == mkvparser::Track::kAudio)
+		{
+			if (FCStringAnsi::Strcmp(Track->GetCodecId(), "A_OPUS") == 0)
+			{
+				MkvFile->AudioTracks.Add(static_cast<const mkvparser::AudioTrack*>(Track));
+
+				// Set as default
+				if (SelectedAudioTrack == INDEX_NONE)
+				{
+					SelectedAudioTrack = MkvFile->AudioTracks.Num() - 1;
+				}
+			}
+			else
+			{
+				UE_LOG(LogWebMMedia, Warning, TEXT("File contains unsupported audio track %d: %s"), i, UTF8_TO_TCHAR(Track->GetCodecId()));
+				continue;
+			}
+		}
+		else
+		{
+			UE_LOG(LogWebMMedia, Warning, TEXT("File contains unsupported track %d: %s"), i, UTF8_TO_TCHAR(Track->GetCodecId()));
+			continue;
+		}
+	}
+
+	if (MkvFile->VideoTracks.Num() == 0 || MkvFile->AudioTracks.Num() == 0)
+	{
+		UE_LOG(LogWebMMedia, Warning, TEXT("File doesn't have video or audio. Right now only files with both are supported"));
+		return false;
+	}
+
+	CurrentTime = 0;
+
+	return true;
+}
+
+void FWebMContainer::ReadFrames(FTimespan ReadBufferLength, TArray<TSharedPtr<FWebMFrame>>& AudioFrames, TArray<TSharedPtr<FWebMFrame>>& VideoFrames)
+{
+	if (bNoMoreToRead)
+	{
+		return;
+	}
+
+	FTimespan UpToReadTime = CurrentTime + ReadBufferLength;
+
+	while (CurrentTime < UpToReadTime)
+	{
+		SeekToNextValidBlock();
+
+		if (bNoMoreToRead)
+		{
+			break;
+		}
+
+		CurrentTime = FTimespan::FromMicroseconds(MkvFile->CurrentBlockEntry->GetBlock()->GetTime(MkvFile->CurrentCluster) / 1000.0);
+
+		for (int32 i = 0; i < MkvFile->CurrentBlockEntry->GetBlock()->GetFrameCount(); ++i)
+		{
+			const mkvparser::Block::Frame& MkvFrame = MkvFile->CurrentBlockEntry->GetBlock()->GetFrame(i);
+
+			TSharedPtr<FWebMFrame> Frame = MakeShared<FWebMFrame>();
+			Frame->Time = CurrentTime;
+			Frame->Data.SetNumUninitialized(MkvFrame.len);
+			MkvFrame.Read(MkvReader.Get(), Frame->Data.GetData());
+
+			int32 TrackNumber = MkvFile->CurrentBlockEntry->GetBlock()->GetTrackNumber();
+
+			if (SelectedVideoTrack != INDEX_NONE && MkvFile->VideoTracks[SelectedVideoTrack]->GetNumber() == TrackNumber)
+			{
+				VideoFrames.Add(Frame);
+				CurrentTime += FTimespan::FromSeconds(1.0 / MkvFile->VideoTracks[SelectedVideoTrack]->GetFrameRate());
+			}
+			else if (SelectedAudioTrack != INDEX_NONE && MkvFile->AudioTracks[SelectedAudioTrack]->GetNumber() == TrackNumber)
+			{
+				AudioFrames.Add(Frame);
+			}
+		}
+	}
+}
+
+FWebMAudioTrackInfo FWebMContainer::GetCurrentAudioTrackInfo() const
+{
+	FWebMAudioTrackInfo Info;
+	if (SelectedAudioTrack == INDEX_NONE)
+	{
+		Info.bIsValid = false;
+		return Info;
+	}
+
+	const mkvparser::AudioTrack* AudioTrack = MkvFile->AudioTracks[SelectedAudioTrack];
+
+	Info.bIsValid = true;
+	Info.CodecName = AudioTrack->GetCodecId();
+	Info.CodecPrivateData = AudioTrack->GetCodecPrivate(Info.CodecPrivateDataSize);
+	Info.SampleRate = AudioTrack->GetSamplingRate();
+	Info.NumOfChannels = AudioTrack->GetChannels();
+
+	return Info;
+}
+
+FWebMVideoTrackInfo FWebMContainer::GetCurrentVideoTrackInfo() const
+{
+	FWebMVideoTrackInfo Info;
+	if (SelectedVideoTrack == INDEX_NONE)
+	{
+		Info.bIsValid = false;
+		return Info;
+	}
+
+	const mkvparser::VideoTrack* VideoTrack = MkvFile->VideoTracks[SelectedVideoTrack];
+
+	Info.bIsValid = true;
+	Info.CodecName = VideoTrack->GetCodecId();
+
+	return Info;
+}
+
+void FWebMContainer::SeekToNextValidBlock()
+{
+	for (;;)
+	{
+		if (!MkvFile->CurrentCluster)
+		{
+			MkvFile->CurrentCluster = MkvFile->CurrentSegment->GetFirst();
+			MkvFile->CurrentBlockEntry = nullptr;
+			check(MkvFile->CurrentCluster);
+		}
+
+		if (!MkvFile->CurrentBlockEntry || MkvFile->CurrentBlockEntry->EOS())
+		{
+			if (MkvFile->CurrentCluster->GetFirst(MkvFile->CurrentBlockEntry) != 0)
+			{
+				UE_LOG(LogWebMMedia, Warning, TEXT("Something went wrong while seeking"));
+				bNoMoreToRead = true;
+				return;
+			}
+		}
+		else
+		{
+			if (MkvFile->CurrentCluster->GetNext(MkvFile->CurrentBlockEntry, MkvFile->CurrentBlockEntry) != 0)
+			{
+				UE_LOG(LogWebMMedia, Warning, TEXT("Something went wrong while seeking"));
+				bNoMoreToRead = true;
+				return;
+			}
+		}
+
+		if (!MkvFile->CurrentBlockEntry || MkvFile->CurrentBlockEntry->EOS())
+		{
+			MkvFile->CurrentBlockEntry = nullptr;
+			MkvFile->CurrentCluster = MkvFile->CurrentSegment->GetNext(MkvFile->CurrentCluster);
+			check(MkvFile->CurrentCluster);
+
+			if (MkvFile->CurrentCluster->EOS())
+			{
+				bNoMoreToRead = true;
+				return;
+			}
+			else
+			{
+				continue;
+			}
+		}
+
+		return;
+	}
+}
diff --git a/Engine/Plugins/Media/WebMMedia/Source/WebMMedia/Public/WebMContainer.h b/Engine/Plugins/Media/WebMMedia/Source/WebMMedia/Public/WebMContainer.h
new file mode 100644
--- /dev/null
+++ b/Engine/Plugins/Media/WebMMedia/Source/WebMMedia/Public/WebMContainer.h
@@ -0,0 +1,50 @@
+// Copyright 1998-2018 Epic Games, Inc. All Rights Reserved.
+
+#pragma once
+
+#include "CoreMinimal.h"
+#include "Containers/Array.h"
+#include "Templates/SharedPointer.h"
+
+struct FWebMFrame;
+
+struct WEBMMEDIA_API FWebMAudioTrackInfo
+{
+	const char* CodecName;
+	const uint8* CodecPrivateData;
+	size_t CodecPrivateDataSize;
+	int32 SampleRate;
+	int32 NumOfChannels;
+	bool bIsValid;
+};
+
+struct WEBMMEDIA_API FWebMVideoTrackInfo
+{
+	const char* CodecName;
+	bool bIsValid;
+};
+
+class FMkvFileReader;
+
+class WEBMMEDIA_API FWebMContainer
+{
+public:
+	FWebMContainer();
+	virtual ~FWebMContainer();
+
+	bool Open(const FString& File);
+	void ReadFrames(FTimespan AmountOfTimeToRead, TArray<TSharedPtr<FWebMFrame>>& AudioFrames, TArray<TSharedPtr<FWebMFrame>>& VideoFrames);
+	FWebMAudioTrackInfo GetCurrentAudioTrackInfo() const;
+	FWebMVideoTrackInfo GetCurrentVideoTrackInfo() const;
+
+private:
+	struct FMkvFileState;
+	TUniquePtr<FMkvFileReader> MkvReader;
+	TUniquePtr<FMkvFileState> MkvFile;
+	FTimespan CurrentTime;
+	int32 SelectedAudioTrack;
+	int32 SelectedVideoTrack;
+	bool bNoMoreToRead;
+
+	void SeekToNextValidBlock();
+};
diff --git a/Engine/Plugins/Media/WebMMedia/Source/WebMMedia/Public/WebMVideoDecoder.h b/Engine/Plugins/Media/WebMMedia/Source/WebMMedia/Public/WebMVideoDecoder.h
--- a/Engine/Plugins/Media/WebMMedia/Source/WebMMedia/Public/WebMVideoDecoder.h
+++ b/Engine/Plugins/Media/WebMMedia/Source/WebMMedia/Public/WebMVideoDecoder.h
@@ -24,6 +24,7 @@ public:
 public:
 	bool Initialize(const char* CodecName);
 	void DecodeVideoFramesAsync(const TArray<TSharedPtr<FWebMFrame>>& VideoFrames);
+	bool IsBusy() const;
 
 private:
 	struct FConvertParams
diff --git a/Engine/Plugins/Media/WebMMedia/Source/WebMMedia/WebMMedia.Build.cs b/Engine/Plugins/Media/WebMMedia/Source/WebMMedia/WebMMedia.Build.cs
--- a/Engine/Plugins/Media/WebMMedia/Source/WebMMedia/WebMMedia.Build.cs
+++ b/Engine/Plugins/Media/WebMMedia/Source/WebMMedia/WebMMedia.Build.cs
@@ -12,11 +12,15 @@ namespace UnrealBuildTool.Rules
 					"WebMMediaFactory",
 					"Core",
 					"Engine",
-					"Media",
-					"MediaUtils",
 					"RenderCore",
 					"RHI",
 					"ShaderCore",
+				});
+
+			PublicDependencyModuleNames.AddRange(
+				new string[] {
+					"Media",
+					"MediaUtils",
 					"UtilityShaders",
 					"libOpus",
 					"UEOgg",
diff --git a/Engine/Plugins/Runtime/WebMMoviePlayer/Source/WebMMoviePlayer/Private/WebMMovieCommon.h b/Engine/Plugins/Runtime/WebMMoviePlayer/Source/WebMMoviePlayer/Private/WebMMovieCommon.h
new file mode 100644
--- /dev/null
+++ b/Engine/Plugins/Runtime/WebMMoviePlayer/Source/WebMMoviePlayer/Private/WebMMovieCommon.h
@@ -0,0 +1,7 @@
+// Copyright 1998-2018 Epic Games, Inc. All Rights Reserved.
+
+#pragma once
+
+#include "Logging/LogMacros.h"
+
+DECLARE_LOG_CATEGORY_EXTERN(LogWebMMoviePlayer, Log, All);
diff --git a/Engine/Plugins/Runtime/WebMMoviePlayer/Source/WebMMoviePlayer/Private/WebMMoviePlayer.cpp b/Engine/Plugins/Runtime/WebMMoviePlayer/Source/WebMMoviePlayer/Private/WebMMoviePlayer.cpp
new file mode 100644
--- /dev/null
+++ b/Engine/Plugins/Runtime/WebMMoviePlayer/Source/WebMMoviePlayer/Private/WebMMoviePlayer.cpp
@@ -0,0 +1,26 @@
+// Copyright 1998-2018 Epic Games, Inc. All Rights Reserved.
+
+#include "WebMMovieStreamer.h"
+
+#include "MoviePlayer.h"
+#include "Modules/ModuleManager.h"
+
+TSharedPtr<FWebMMovieStreamer> WebMMovieStreamer;
+
+class FWebMMoviePlayerModule : public IModuleInterface
+{
+	/** IModuleInterface implementation */
+	virtual void StartupModule() override
+	{
+		FWebMMovieStreamer *Streamer = new FWebMMovieStreamer;
+		WebMMovieStreamer = MakeShareable(Streamer);
+		GetMoviePlayer()->RegisterMovieStreamer(WebMMovieStreamer);
+	}
+
+	virtual void ShutdownModule() override
+	{
+		WebMMovieStreamer.Reset();
+	}
+};
+
+IMPLEMENT_MODULE(FWebMMoviePlayerModule, WebMMoviePlayer)
diff --git a/Engine/Plugins/Runtime/WebMMoviePlayer/Source/WebMMoviePlayer/Private/WebMMovieStreamer.cpp b/Engine/Plugins/Runtime/WebMMoviePlayer/Source/WebMMoviePlayer/Private/WebMMovieStreamer.cpp
new file mode 100644
--- /dev/null
+++ b/Engine/Plugins/Runtime/WebMMoviePlayer/Source/WebMMoviePlayer/Private/WebMMovieStreamer.cpp
@@ -0,0 +1,236 @@
+// Copyright 1998-2018 Epic Games, Inc. All Rights Reserved.
+
+#include "WebMMovieStreamer.h"
+#include "MediaShaders.h"
+#include "MediaSamples.h"
+#include "WebMMovieCommon.h"
+#include "Misc/Paths.h"
+#include "SceneUtils.h"
+#include "StaticBoundShaderState.h"
+#include "PipelineStateCache.h"
+#include "RHIStaticStates.h"
+#include "WebMVideoDecoder.h"
+#include "WebMAudioDecoder.h"
+#include "WebMMediaFrame.h"
+#include "WebMContainer.h"
+#include "WebMMediaTextureSample.h"
+
+DEFINE_LOG_CATEGORY(LogWebMMoviePlayer);
+
+FWebMMovieStreamer::FWebMMovieStreamer()
+	: CurrentTexture(0)
+	, Samples(new FMediaSamples())
+	, CurrentTime(0)
+	, bPlaying(false)
+{
+	Viewport = MakeShareable(new FMovieViewport());
+}
+
+FWebMMovieStreamer::~FWebMMovieStreamer()
+{
+	Cleanup();
+}
+
+void FWebMMovieStreamer::Cleanup()
+{
+	ForceCompletion();
+
+	ReleaseAcquiredResources();
+}
+
+FTexture2DRHIRef FWebMMovieStreamer::GetTexture()
+{
+	return BufferedVideoTextures[CurrentTexture].IsValid() ? BufferedVideoTextures[CurrentTexture]->GetRHIRef() : nullptr;
+}
+
+bool FWebMMovieStreamer::Init(const TArray<FString>& InMoviePaths, TEnumAsByte<EMoviePlaybackType> InPlaybackType)
+{
+	// Initializes the streamer for audio and video playback of the given path(s).
+	// NOTE: If multiple paths are provided, it is expect that they be played back seamlessly.
+
+	// Add the given paths to the movie queue
+	MovieQueue.Append(InMoviePaths);
+
+	// start our first movie playing
+	bPlaying = StartNextMovie();
+
+	// Play the next movie in the queue
+	return bPlaying;
+
+}
+
+bool FWebMMovieStreamer::StartNextMovie()
+{
+	if (MovieQueue.Num() > 0)
+	{
+		MovieName = MovieQueue[0];
+
+		MovieQueue.RemoveAt(0);
+
+		FString MoviePath = FPaths::ProjectContentDir() + TEXT("Movies/") + MovieName + TEXT(".webm");
+
+		if (FPaths::FileExists(MoviePath))
+		{
+			Container.Reset(new FWebMContainer());
+		}
+		else
+		{
+			UE_LOG(LogWebMMoviePlayer, Error, TEXT("Movie '%s' not found."));
+		}
+
+		UE_LOG(LogWebMMoviePlayer, Verbose, TEXT("Starting '%s'"), *MoviePath);
+
+		if (!Container->Open(MoviePath))
+		{
+			MovieName = FString();
+			return false;
+		}
+
+		AudioDecoder.Reset(new FWebMAudioDecoder(Samples));
+		VideoDecoder.Reset(new FWebMVideoDecoder(Samples));
+
+		FWebMAudioTrackInfo DefaultAudioTrack = Container->GetCurrentAudioTrackInfo();
+		check(DefaultAudioTrack.bIsValid);
+
+		FWebMVideoTrackInfo DefaultVideoTrack = Container->GetCurrentVideoTrackInfo();
+		check(DefaultVideoTrack.bIsValid);
+
+		AudioDecoder->Initialize(DefaultAudioTrack.CodecName, DefaultAudioTrack.SampleRate, DefaultAudioTrack.NumOfChannels, DefaultAudioTrack.CodecPrivateData, DefaultAudioTrack.CodecPrivateDataSize);
+		VideoDecoder->Initialize(DefaultVideoTrack.CodecName);
+
+		CurrentTime = 0;
+
+		return true;
+	}
+	else
+	{
+		UE_LOG(LogWebMMoviePlayer, Verbose, TEXT("No Movie to start."));
+		return false;
+	}
+}
+
+FString FWebMMovieStreamer::GetMovieName()
+{
+	return MovieName;
+}
+
+bool FWebMMovieStreamer::IsLastMovieInPlaylist()
+{
+	return MovieQueue.Num() == 0;
+}
+
+bool FWebMMovieStreamer::Tick(float InDeltaTime)
+{
+	if (bPlaying)
+	{
+		bool bHaveThingsToDo = false;
+
+		bHaveThingsToDo |= DisplayFrames(InDeltaTime);
+
+		// TODO! Audio!
+
+		bHaveThingsToDo |= DecodeMoreFrames();
+
+		CurrentTime += InDeltaTime;
+
+		return !bHaveThingsToDo;
+	}
+	else
+	{
+		// We're done playing
+		return true;
+	}
+}
+
+TSharedPtr<class ISlateViewport> FWebMMovieStreamer::GetViewportInterface()
+{
+	return Viewport;
+}
+
+float FWebMMovieStreamer::GetAspectRatio() const
+{
+	return static_cast<float>(Viewport->GetSize().X) / static_cast<float>(Viewport->GetSize().Y);
+}
+
+void FWebMMovieStreamer::ForceCompletion()
+{
+	if (bPlaying)
+	{
+		bPlaying = false;
+	}
+
+	MovieQueue.Empty();
+}
+
+void FWebMMovieStreamer::ReleaseAcquiredResources()
+{
+	// NOTE: Called from the main thread
+	for (int32 i = 0; i < NumBufferedTextures; ++i)
+	{
+		TSharedPtr<FSlateTexture2DRHIRef, ESPMode::ThreadSafe>& VideoTexture = BufferedVideoTextures[i];
+		if (VideoTexture.IsValid())
+		{
+			// Schedule the release of the video texture(s) and wait the release to complete before
+			// losing the reference to the texture
+			BeginReleaseResource(VideoTexture.Get());
+			FlushRenderingCommands();
+			VideoTexture.Reset();
+		}
+	}
+
+	Viewport->SetTexture(nullptr);
+
+	VideoDecoder.Reset();
+	AudioDecoder.Reset();
+	Container.Reset();
+}
+
+bool FWebMMovieStreamer::DisplayFrames(float InDeltaTime)
+{
+	TSharedPtr<IMediaTextureSample, ESPMode::ThreadSafe> VideoSample;
+	TRange<FTimespan> TimeRange(FTimespan::Zero(), FTimespan::FromSeconds(CurrentTime));
+	bool bFoundSample = Samples->FetchVideo(TimeRange, VideoSample);
+
+	if (bFoundSample && VideoSample)
+	{
+		CurrentTexture = (CurrentTexture + 1) % NumBufferedTextures;
+
+		FWebMMediaTextureSample* WebMSample = StaticCast<FWebMMediaTextureSample*>(VideoSample.Get());
+
+		if (!BufferedVideoTextures[CurrentTexture].IsValid())
+		{
+			BufferedVideoTextures[CurrentTexture] = MakeShareable(new FSlateTexture2DRHIRef(nullptr, 0, 0));
+		}
+
+		BufferedVideoTextures[CurrentTexture]->SetRHIRef(WebMSample->GetTextureRef(), WebMSample->GetDim().X, WebMSample->GetDim().Y);
+
+		Viewport->SetTexture(BufferedVideoTextures[CurrentTexture]);
+	}
+
+	return Samples->NumVideoSamples() > 0 || VideoDecoder->IsBusy();
+}
+
+bool FWebMMovieStreamer::DecodeMoreFrames()
+{
+	// Read frames up to 1 secs in the future
+	FTimespan ReadBufferLength = FTimespan::FromSeconds(0.1);
+
+	TArray<TSharedPtr<FWebMFrame>> AudioFrames;
+	TArray<TSharedPtr<FWebMFrame>> VideoFrames;
+
+	Container->ReadFrames(ReadBufferLength, AudioFrames, VideoFrames);
+
+	// Trigger video decoding
+	if (VideoFrames.Num() > 0)
+	{
+		VideoDecoder->DecodeVideoFramesAsync(VideoFrames);
+	}
+
+	// Trigger audio decoding
+	if (AudioFrames.Num() > 0)
+	{
+		AudioDecoder->DecodeAudioFramesAsync(AudioFrames);
+	}
+
+	return VideoFrames.Num() > 0 || AudioFrames.Num() > 0;
+}
diff --git a/Engine/Plugins/Runtime/WebMMoviePlayer/Source/WebMMoviePlayer/Private/WebMMovieStreamer.h b/Engine/Plugins/Runtime/WebMMoviePlayer/Source/WebMMoviePlayer/Private/WebMMovieStreamer.h
new file mode 100644
--- /dev/null
+++ b/Engine/Plugins/Runtime/WebMMoviePlayer/Source/WebMMoviePlayer/Private/WebMMovieStreamer.h
@@ -0,0 +1,55 @@
+// Copyright 1998-2018 Epic Games, Inc. All Rights Reserved.
+
+#pragma once
+
+#include "MoviePlayer.h"
+
+class FWebMVideoDecoder;
+class FWebMAudioDecoder;
+class FMediaSamples;
+class FWebMContainer;
+
+class FWebMMovieStreamer : public IMovieStreamer
+{
+public:
+	FWebMMovieStreamer();
+	~FWebMMovieStreamer();
+
+	//~ IMediaOutput interface
+	virtual bool Init(const TArray<FString>& InMoviePaths, TEnumAsByte<EMoviePlaybackType> InPlaybackType) override;
+	virtual void ForceCompletion() override;
+	virtual bool Tick(float InDeltaTime) override;
+	virtual TSharedPtr<class ISlateViewport> GetViewportInterface() override;
+	virtual float GetAspectRatio() const override;
+	virtual FString GetMovieName() override;
+	virtual bool IsLastMovieInPlaylist() override;
+	virtual void Cleanup() override;
+	virtual FTexture2DRHIRef GetTexture() override;
+
+	FOnCurrentMovieClipFinished OnCurrentMovieClipFinishedDelegate;
+	virtual FOnCurrentMovieClipFinished& OnCurrentMovieClipFinished() override { return OnCurrentMovieClipFinishedDelegate; }
+
+private:
+	//Theoretically only need 2 buffered textures, but we have extra to avoid needing to make a copy of the AvPlayer data to pass to an RHI thread command.  Instead, we buffer deeper and update the textures on the Render thread.
+	static const int32 NumBufferedTextures = 4;
+
+	/** Texture and viewport data for displaying to Slate.  SlateVideoTexture is always used by the viewport, but it's texture reference is swapped out when a new frame is available.
+	 Method assumes 1 Tick() call per frame, and that the Streamer Tick comes before Slate rendering */
+	TSharedPtr<FSlateTexture2DRHIRef, ESPMode::ThreadSafe> BufferedVideoTextures[NumBufferedTextures];
+	int32 CurrentTexture;
+
+	TArray<FString> MovieQueue;
+	FString MovieName;
+	TUniquePtr<FWebMVideoDecoder> VideoDecoder;
+	TUniquePtr<FWebMAudioDecoder> AudioDecoder;
+	TUniquePtr<FWebMContainer> Container;
+	TSharedPtr<FMediaSamples, ESPMode::ThreadSafe> Samples;
+	TSharedPtr<FMovieViewport> Viewport;
+	float CurrentTime;
+	bool bPlaying;
+
+	bool StartNextMovie();
+	void ReleaseAcquiredResources();
+	bool DisplayFrames(float InDeltaTime);
+	bool DecodeMoreFrames();
+};
diff --git a/Engine/Plugins/Runtime/WebMMoviePlayer/Source/WebMMoviePlayer/WebMMoviePlayer.Build.cs b/Engine/Plugins/Runtime/WebMMoviePlayer/Source/WebMMoviePlayer/WebMMoviePlayer.Build.cs
new file mode 100644
--- /dev/null
+++ b/Engine/Plugins/Runtime/WebMMoviePlayer/Source/WebMMoviePlayer/WebMMoviePlayer.Build.cs
@@ -0,0 +1,26 @@
+// Copyright 1998-2018 Epic Games, Inc. All Rights Reserved.
+using System;
+using System.IO;
+
+namespace UnrealBuildTool.Rules
+{
+	public class WebMMoviePlayer : ModuleRules
+	{
+		public WebMMoviePlayer(ReadOnlyTargetRules Target) : base(Target)
+		{
+			PrivateDependencyModuleNames.AddRange(
+				new string[]
+				{
+					"Core",
+					"Engine",
+					"MoviePlayer",
+					"RenderCore",
+					"RHI",
+					"SlateCore",
+					"Slate",
+					"MediaUtils",
+					"WebMMedia",
+				});
+		}
+	}
+}
diff --git a/Engine/Plugins/Runtime/WebMMoviePlayer/WebMMoviePlayer.uplugin b/Engine/Plugins/Runtime/WebMMoviePlayer/WebMMoviePlayer.uplugin
new file mode 100644
--- /dev/null
+++ b/Engine/Plugins/Runtime/WebMMoviePlayer/WebMMoviePlayer.uplugin
@@ -0,0 +1,42 @@
+{
+	"FileVersion": 3,
+	"Version": 1,
+	"VersionName": "1.0",
+	"FriendlyName": "Movie Player for WebM files",
+	"Description": "Movie Player for WebM files",
+	"Category": "Movie Players",
+	"CreatedBy": "Epic Games, Inc.",
+	"CreatedByURL": "http://epicgames.com",
+	"DocsURL": "",
+	"MarketplaceURL": "",
+	"SupportURL": "",
+	"EnabledByDefault": true,
+	"CanContainContent": false,
+	"IsBetaVersion": false,
+	"Installed": false,
+	"Modules": [
+		{
+			"Name": "WebMMoviePlayer",
+			"Type": "RuntimeNoCommandlet",
+			"LoadingPhase": "PreLoadingScreen",
+			"BlacklistPlatforms": [
+				"Win32",
+				"Mac",
+				"XboxOne",
+				"PS4",
+				"IOS",
+				"Android",
+				"HTML5",
+				"TVOS",
+				"Switch",
+				"Lumin"
+			]
+		}
+	],
+	"Plugins": [
+		{
+			"Name": "WebMMedia",
+			"Enabled": true
+		}
+	]
+}
diff --git a/Engine/Source/Runtime/MoviePlayer/Private/MoviePlayer.cpp b/Engine/Source/Runtime/MoviePlayer/Private/MoviePlayer.cpp
--- a/Engine/Source/Runtime/MoviePlayer/Private/MoviePlayer.cpp
+++ b/Engine/Source/Runtime/MoviePlayer/Private/MoviePlayer.cpp
@@ -125,9 +125,5 @@ bool IsMoviePlayerEnabled()
 	bEnabled &= !FParse::Param(FCommandLine::Get(), TEXT("NoLoadingScreen"));
 #endif
 
-#if PLATFORM_UNIX
-	bEnabled = false;
-#endif
-
 	return bEnabled;
 }
