Linux: remove webm from the server builds.

diff --git a/Engine/Plugins/Media/WebMMedia/WebMMedia.uplugin b/Engine/Plugins/Media/WebMMedia/WebMMedia.uplugin
--- a/Engine/Plugins/Media/WebMMedia/WebMMedia.uplugin
+++ b/Engine/Plugins/Media/WebMMedia/WebMMedia.uplugin
@@ -28,6 +28,10 @@
 				"TVOS",
 				"Switch",
 				"Lumin",
+			],
+			"BlacklistTargets" :
+			[
+				"Server",
 			]
 		},
 		{
@@ -56,6 +60,10 @@
 				"TVOS",
 				"Switch",
 				"Lumin",
+			],
+			"BlacklistTargets" :
+			[
+				"Server",
 			]
 		}
 	]
diff --git a/Engine/Plugins/Runtime/WebMMoviePlayer/WebMMoviePlayer.uplugin b/Engine/Plugins/Runtime/WebMMoviePlayer/WebMMoviePlayer.uplugin
--- a/Engine/Plugins/Runtime/WebMMoviePlayer/WebMMoviePlayer.uplugin
+++ b/Engine/Plugins/Runtime/WebMMoviePlayer/WebMMoviePlayer.uplugin
@@ -29,6 +29,10 @@
 				"TVOS",
 				"Switch",
 				"Lumin"
+			],
+			"BlacklistTargets" :
+			[
+				"Server",
 			]
 		}
 	],
