DR - Fix deprecation warnings

diff --git a/Engine/Plugins/Media/WebMMedia/Source/WebMMedia/Private/Player/WebMVideoDecoder.cpp b/Engine/Plugins/Media/WebMMedia/Source/WebMMedia/Private/Player/WebMVideoDecoder.cpp
--- a/Engine/Plugins/Media/WebMMedia/Source/WebMMedia/Private/Player/WebMVideoDecoder.cpp
+++ b/Engine/Plugins/Media/WebMMedia/Source/WebMMedia/Private/Player/WebMVideoDecoder.cpp
@@ -232,7 +232,7 @@ void FWebMVideoDecoder::ConvertYUVToRGBAndSubmit(const FConvertParams& Params)
 		// draw full-size quad
 		CommandList.SetViewport(0, 0, 0.0f, Image->w, Image->d_h, 1.0f);
 		CommandList.SetStreamSource(0, GMoviePlayerResources.VertexBufferRHI, 0);
-		CommandList.DrawPrimitive(PT_TriangleStrip, 0, 2, 1);
+		CommandList.DrawPrimitive(0, 2, 1);
 		CommandList.CopyToResolveTarget(RenderTarget, RenderTarget, FResolveParams());
 
 		Samples->AddVideo(VideoSample.ToSharedRef());
diff --git a/Engine/Source/Runtime/Engine/Private/Particles/ParticleGpuSimulation.cpp b/Engine/Source/Runtime/Engine/Private/Particles/ParticleGpuSimulation.cpp
--- a/Engine/Source/Runtime/Engine/Private/Particles/ParticleGpuSimulation.cpp
+++ b/Engine/Source/Runtime/Engine/Private/Particles/ParticleGpuSimulation.cpp
@@ -1437,7 +1437,6 @@ static void DrawParticleTiles(FRHICommandList& RHICmdList, FParticleBufferParamR
 	// Draw tiles.
 	RHICmdList.DrawIndexedPrimitive(
 		GParticleIndexBuffer.IndexBufferRHI,
-		PT_TriangleList,
 		/*BaseVertexIndex=*/0,
 		/*MinIndex=*/ 0,
 		/*NumVertices=*/ 4,
diff --git a/Engine/Source/Runtime/UtilityShaders/Private/ClearQuad.cpp b/Engine/Source/Runtime/UtilityShaders/Private/ClearQuad.cpp
--- a/Engine/Source/Runtime/UtilityShaders/Private/ClearQuad.cpp
+++ b/Engine/Source/Runtime/UtilityShaders/Private/ClearQuad.cpp
@@ -279,7 +279,7 @@ void DrawClearQuadMRT(FRHICommandList& RHICmdList, bool bClearColor, int32 NumCl
 
 	// Draw a fullscreen quad without a hole
 	RHICmdList.SetStreamSource(0, GClearVertexBuffer.VertexBufferRHI, 0);
-	RHICmdList.DrawPrimitive(PT_TriangleStrip, 0, 2, 1);
+	RHICmdList.DrawPrimitive(0, 2, 1);
 
 	if (ClearQuadCallbacks.PostClear)
 	{
