DR - Merge fix
#rnx

diff --git a/Engine/Plugins/Media/WebMMedia/Source/WebMMedia/Private/Player/WebMVideoDecoder.cpp b/Engine/Plugins/Media/WebMMedia/Source/WebMMedia/Private/Player/WebMVideoDecoder.cpp
--- a/Engine/Plugins/Media/WebMMedia/Source/WebMMedia/Private/Player/WebMVideoDecoder.cpp
+++ b/Engine/Plugins/Media/WebMMedia/Source/WebMMedia/Private/Player/WebMVideoDecoder.cpp
@@ -68,13 +68,7 @@ FWebMVideoDecoder::FWebMVideoDecoder(TSharedPtr<FMediaSamples, ESPMode::ThreadSa
 
 FWebMVideoDecoder::~FWebMVideoDecoder()
 {
->>>> ORIGINAL //UE4/Main/Engine/Plugins/Media/WebMMedia/Source/WebMMedia/Private/Player/WebMVideoDecoder.cpp#1
-	// Make sure all compute shader decoding is done
-	FlushRenderingCommands();
-==== THEIRS //UE4/Main/Engine/Plugins/Media/WebMMedia/Source/WebMMedia/Private/Player/WebMVideoDecoder.cpp#3
 	Close();
-==== YOURS //rolando.caloca_I9_DR2/Engine/Plugins/Media/WebMMedia/Source/WebMMedia/Private/Player/WebMVideoDecoder.cpp
-<<<<
 }
 
 bool FWebMVideoDecoder::Initialize(const char* CodecName)
@@ -97,18 +91,7 @@ bool FWebMVideoDecoder::Initialize(const char* CodecName)
 		return false;
 	}
 
->>>> ORIGINAL //UE4/Main/Engine/Plugins/Media/WebMMedia/Source/WebMMedia/Private/Player/WebMVideoDecoder.cpp#1
-	vpx_codec_destroy(&Context);
-}
-==== THEIRS //UE4/Main/Engine/Plugins/Media/WebMMedia/Source/WebMMedia/Private/Player/WebMVideoDecoder.cpp#3
 	bIsInitialized = true;
-==== YOURS //rolando.caloca_I9_DR2/Engine/Plugins/Media/WebMMedia/Source/WebMMedia/Private/Player/WebMVideoDecoder.cpp
-	// Make sure all compute shader decoding is done
-	FlushRenderingCommands();
-
-	vpx_codec_destroy(&Context);
-}
-<<<<
 
 	return true;
 }
