More changes from Fury that had been cherry-picked so had unresolved integrations to identical files.

Plus WebMMedia.uplugin which seems to have malformed json in the Main/Dev-VP branches...

diff --git a/Engine/Plugins/Media/WebMMedia/WebMMedia.uplugin b/Engine/Plugins/Media/WebMMedia/WebMMedia.uplugin
--- a/Engine/Plugins/Media/WebMMedia/WebMMedia.uplugin
+++ b/Engine/Plugins/Media/WebMMedia/WebMMedia.uplugin
@@ -31,7 +31,7 @@
 			],
 			"BlacklistTargets" :
 			[
-				"Server",
+				"Server"
 			]
 		},
 		{
@@ -63,7 +63,7 @@
 			],
 			"BlacklistTargets" :
 			[
-				"Server",
+				"Server"
 			]
 		}
 	]
