Fix macOS build.

#rnx

diff --git a/Engine/Plugins/Media/WebMMedia/WebMMedia.uplugin b/Engine/Plugins/Media/WebMMedia/WebMMedia.uplugin
--- a/Engine/Plugins/Media/WebMMedia/WebMMedia.uplugin
+++ b/Engine/Plugins/Media/WebMMedia/WebMMedia.uplugin
@@ -19,7 +19,8 @@
 			"LoadingPhase": "PostEngineInit",
 			"WhitelistPlatforms": [
 				"Win64",
-				"Linux"
+				"Linux",
+				"Mac"
 			],
 			"BlacklistTargets": [
 				"Server"
@@ -41,7 +42,8 @@
 			"LoadingPhase": "PostEngineInit",
 			"WhitelistPlatforms": [
 				"Win64",
-				"Linux"
+				"Linux",
+				"Mac"
 			],
 			"BlacklistTargets": [
 				"Server"
diff --git a/Engine/Plugins/Runtime/WebMMoviePlayer/WebMMoviePlayer.uplugin b/Engine/Plugins/Runtime/WebMMoviePlayer/WebMMoviePlayer.uplugin
--- a/Engine/Plugins/Runtime/WebMMoviePlayer/WebMMoviePlayer.uplugin
+++ b/Engine/Plugins/Runtime/WebMMoviePlayer/WebMMoviePlayer.uplugin
@@ -21,7 +21,8 @@
 			"LoadingPhase": "PreLoadingScreen",
 			"WhitelistPlatforms": [
 				"Win64",
-				"Linux"
+				"Linux",
+				"Mac"
 			],
 			"BlacklistTargets": [
 				"Server"
