Webm: fixed cropped rendering when internal decoder resolution does not match actual video output resolution

The vpx_image knowss these fields from the decoder and there could be a difference between image storage and image display dimentsion.

/* Image storage dimensions */
unsigned int w;    /**< Stored image width */
unsigned int h;    /**< Stored image height */

/* Image display dimensions */
unsigned int d_w; /**< Displayed image width */
unsigned int d_h; /**< Displayed image height */

This fix uses the displayed image width as actual rendering output width.

diff --git a/Engine/Plugins/Media/WebMMedia/Source/WebMMedia/Private/Player/WebMVideoDecoder.cpp b/Engine/Plugins/Media/WebMMedia/Source/WebMMedia/Private/Player/WebMVideoDecoder.cpp
--- a/Engine/Plugins/Media/WebMMedia/Source/WebMMedia/Private/Player/WebMVideoDecoder.cpp
+++ b/Engine/Plugins/Media/WebMMedia/Source/WebMMedia/Private/Player/WebMVideoDecoder.cpp
@@ -271,7 +271,7 @@ void FWebMVideoDecoder::ConvertYUVToRGBAndSubmit(const FConvertParams& Params)
 			PixelShader->SetParameters(CommandList, DecodedY->GetTexture2D(), DecodedU->GetTexture2D(), DecodedV->GetTexture2D(), FIntPoint(Image->d_w, Image->d_h), MediaShaders::YuvToSrgbDefault, MediaShaders::YUVOffset8bits, true);
 
 			// draw full-size quad
-			CommandList.SetViewport(0, 0, 0.0f, Image->w, Image->d_h, 1.0f);
+			CommandList.SetViewport(0, 0, 0.0f, Image->d_w, Image->d_h, 1.0f);
 			CommandList.SetStreamSource(0, GMoviePlayerResources.VertexBufferRHI, 0);
 			CommandList.DrawPrimitive(0, 2, 1);
 		}
