WebM: use single-threaded decoding when running -norenderthread (UE-91327).

- TaskGraph does not respect -norenderthread (formerly -onethread), so this runs against the RHI assumptions.

diff --git a/Engine/Plugins/Media/WebMMedia/Source/WebMMedia/Private/Player/WebMVideoDecoder.cpp b/Engine/Plugins/Media/WebMMedia/Source/WebMMedia/Private/Player/WebMVideoDecoder.cpp
--- a/Engine/Plugins/Media/WebMMedia/Source/WebMMedia/Private/Player/WebMVideoDecoder.cpp
+++ b/Engine/Plugins/Media/WebMMedia/Source/WebMMedia/Private/Player/WebMVideoDecoder.cpp
@@ -101,17 +101,25 @@ bool FWebMVideoDecoder::Initialize(const char* CodecName)
 
 void FWebMVideoDecoder::DecodeVideoFramesAsync(const TArray<TSharedPtr<FWebMFrame>>& VideoFrames)
 {
-	FGraphEventRef PreviousDecodingTask = VideoDecodingTask;
-
-	VideoDecodingTask = FFunctionGraphTask::CreateAndDispatchWhenReady([this, PreviousDecodingTask, VideoFrames]()
+	static bool bUseRenderThread = FPlatformMisc::UseRenderThread();
+	if (bUseRenderThread)
 	{
-		if(PreviousDecodingTask && !PreviousDecodingTask->IsComplete())
+		FGraphEventRef PreviousDecodingTask = VideoDecodingTask;
+
+		VideoDecodingTask = FFunctionGraphTask::CreateAndDispatchWhenReady([this, PreviousDecodingTask, VideoFrames]()
 		{
-			FTaskGraphInterface::Get().WaitUntilTaskCompletes(PreviousDecodingTask);
-		}
+			if(PreviousDecodingTask && !PreviousDecodingTask->IsComplete())
+			{
+				FTaskGraphInterface::Get().WaitUntilTaskCompletes(PreviousDecodingTask);
+			}
 
+			DoDecodeVideoFrames(VideoFrames);
+		}, TStatId(), nullptr, ENamedThreads::AnyThread);
+	}
+	else
+	{
 		DoDecodeVideoFrames(VideoFrames);
-	}, TStatId(), nullptr, ENamedThreads::AnyThread);
+	}
 }
 
 bool FWebMVideoDecoder::IsBusy() const
