Only copy full texture only when RHI stride matches with WebM image stride, otherwise copy line by line

diff --git a/Engine/Plugins/Media/WebMMedia/Source/WebMMedia/Private/Player/WebMVideoDecoder.cpp b/Engine/Plugins/Media/WebMMedia/Source/WebMMedia/Private/Player/WebMVideoDecoder.cpp
--- a/Engine/Plugins/Media/WebMMedia/Source/WebMMedia/Private/Player/WebMVideoDecoder.cpp
+++ b/Engine/Plugins/Media/WebMMedia/Source/WebMMedia/Private/Player/WebMVideoDecoder.cpp
@@ -213,45 +213,43 @@ void FWebMVideoDecoder::ConvertYUVToRGBAndSubmit(const FConvertParams& Params)
 
 	// render video frame into output texture
 	FRHICommandListImmediate& CommandList = GetImmediateCommandList_ForRenderCommand();
+
 	{
-		// copy the Y plane out of the video buffer
+		const auto CopyTextureMemory = [Image](
+			FRHICommandListImmediate& InCommandList,
+			FRHITexture2D* RHITexture,
+			int ImageIndex,
+			int CopyHeight)
 		{
 			uint32 Stride = 0;
-			void * TextureMemory = GDynamicRHI->LockTexture2D_RenderThread(CommandList, DecodedY.GetReference(), 0, RLM_WriteOnly, Stride, false);
+			unsigned char* TextureMemory = (unsigned char*)GDynamicRHI->LockTexture2D_RenderThread(InCommandList, RHITexture, 0, RLM_WriteOnly, Stride, false);
 
 			if (TextureMemory)
 			{
-				check(Stride == Image->stride[0]);
-				memcpy(TextureMemory, Image->planes[0], Image->stride[0] * Image->d_h);
-				GDynamicRHI->UnlockTexture2D_RenderThread(CommandList, DecodedY.GetReference(), 0, false);
+				check(Stride >= (uint32)Image->stride[ImageIndex]);
+				if (Stride == Image->stride[ImageIndex])
+				{
+					FMemory::Memcpy(TextureMemory, Image->planes[0], Image->stride[ImageIndex] * CopyHeight);
+				}
+				else
+				{
+					for (int h = 0; h < CopyHeight; ++h)
+					{
+						FMemory::Memcpy(TextureMemory + h * Stride, Image->planes[ImageIndex] + h * Image->stride[ImageIndex], Image->stride[ImageIndex]);
+					}
+				}
+				GDynamicRHI->UnlockTexture2D_RenderThread(InCommandList, RHITexture, 0, false);
 			}
-		}
+		};
 
-		// copy the U plane out of the video buffer
-		{
-			uint32 Stride = 0;
-			void * TextureMemory = GDynamicRHI->LockTexture2D_RenderThread(CommandList, DecodedU.GetReference(), 0, RLM_WriteOnly, Stride, false);
+		// copy the Y plane out of the video buffer
+		CopyTextureMemory(CommandList, DecodedY.GetReference(), 0, Image->d_h);
 
-			if (TextureMemory)
-			{
-				check(Stride == Image->stride[1]);
-				memcpy(TextureMemory, Image->planes[1], Image->stride[1] * Image->d_h / 2);
-				GDynamicRHI->UnlockTexture2D_RenderThread(CommandList, DecodedU.GetReference(), 0, false);
-			}
-		}
+		// copy the U plane out of the video buffer
+		CopyTextureMemory(CommandList, DecodedU.GetReference(), 1, Image->d_h / 2);
 
 		// copy the V plane out of the video buffer
-		{
-			uint32 Stride = 0;
-			void * TextureMemory = GDynamicRHI->LockTexture2D_RenderThread(CommandList, DecodedV.GetReference(), 0, RLM_WriteOnly, Stride, false);
-
-			if (TextureMemory)
-			{
-				check(Stride == Image->stride[2]);
-				memcpy(TextureMemory, Image->planes[2], Image->stride[2] * Image->d_h / 2);
-				GDynamicRHI->UnlockTexture2D_RenderThread(CommandList, DecodedV.GetReference(), 0, false);
-			}
-		}
+		CopyTextureMemory(CommandList, DecodedV.GetReference(), 2, Image->d_h / 2);
 
 		FRHITexture* RenderTarget = VideoSample->GetTexture();
 		FRHIRenderPassInfo RPInfo(RenderTarget, ERenderTargetActions::Load_Store);
