WebM: fix "green movies" caused by wrong plane being copied to U and V (UE-92233).

diff --git a/Engine/Plugins/Media/WebMMedia/Source/WebMMedia/Private/Player/WebMVideoDecoder.cpp b/Engine/Plugins/Media/WebMMedia/Source/WebMMedia/Private/Player/WebMVideoDecoder.cpp
--- a/Engine/Plugins/Media/WebMMedia/Source/WebMMedia/Private/Player/WebMVideoDecoder.cpp
+++ b/Engine/Plugins/Media/WebMMedia/Source/WebMMedia/Private/Player/WebMVideoDecoder.cpp
@@ -229,7 +229,7 @@ void FWebMVideoDecoder::ConvertYUVToRGBAndSubmit(const FConvertParams& Params)
 				check(Stride >= (uint32)Image->stride[ImageIndex]);
 				if (Stride == Image->stride[ImageIndex])
 				{
-					FMemory::Memcpy(TextureMemory, Image->planes[0], Image->stride[ImageIndex] * CopyHeight);
+					FMemory::Memcpy(TextureMemory, Image->planes[ImageIndex], Image->stride[ImageIndex] * CopyHeight);
 				}
 				else
 				{
