WebMMedia: Fixed hang in player.
#rnx

diff --git a/Engine/Plugins/Media/WebMMedia/Source/WebMMedia/Private/Player/WebMVideoDecoder.cpp b/Engine/Plugins/Media/WebMMedia/Source/WebMMedia/Private/Player/WebMVideoDecoder.cpp
--- a/Engine/Plugins/Media/WebMMedia/Source/WebMMedia/Private/Player/WebMVideoDecoder.cpp
+++ b/Engine/Plugins/Media/WebMMedia/Source/WebMMedia/Private/Player/WebMVideoDecoder.cpp
@@ -104,17 +104,15 @@ void FWebMVideoDecoder::DecodeVideoFramesAsync(const TArray<TSharedPtr<FWebMFram
 	static bool bUseRenderThread = FPlatformMisc::UseRenderThread();
 	if (bUseRenderThread)
 	{
-		FGraphEventRef PreviousDecodingTask = VideoDecodingTask;
-
-		VideoDecodingTask = FFunctionGraphTask::CreateAndDispatchWhenReady([this, PreviousDecodingTask, VideoFrames]()
+		FGraphEventArray Prerequisites;
+		if (VideoDecodingTask)
+		{
+			Prerequisites.Emplace(VideoDecodingTask);
+		}
+		VideoDecodingTask = FFunctionGraphTask::CreateAndDispatchWhenReady([this,  VideoFrames]()
 		{
-			if(PreviousDecodingTask && !PreviousDecodingTask->IsComplete())
-			{
-				FTaskGraphInterface::Get().WaitUntilTaskCompletes(PreviousDecodingTask);
-			}
-
 			DoDecodeVideoFrames(VideoFrames);
-		}, TStatId(), nullptr, ENamedThreads::AnyThread);
+		}, TStatId(), &Prerequisites, ENamedThreads::AnyThread);
 	}
 	else
 	{
