Build fix for Stadia - don't use WebRTC for WebMMedia module on Stadia

diff --git a/Engine/Plugins/Media/WebMMedia/Source/WebMMedia/WebMMedia.Build.cs b/Engine/Plugins/Media/WebMMedia/Source/WebMMedia/WebMMedia.Build.cs
--- a/Engine/Plugins/Media/WebMMedia/Source/WebMMedia/WebMMedia.Build.cs
+++ b/Engine/Plugins/Media/WebMMedia/Source/WebMMedia/WebMMedia.Build.cs
@@ -25,7 +25,20 @@ namespace UnrealBuildTool.Rules
 					"Vorbis",
 				});
 
-			if ((Target.IsInPlatformGroup(UnrealPlatformGroup.Unix) && Target.Architecture.StartsWith("x86_64")) || Target.Platform == UnrealTargetPlatform.Win64)
+			if ((Target.Platform == UnrealTargetPlatform.Stadia) || (Target.Platform == UnrealTargetPlatform.Mac))
+			{
+				// Engine does not currently ship with Mac or Stadia binaries for WebRTC so use fallback
+				
+				PrivateDependencyModuleNames.Add("LibVpx");
+				
+				PublicDependencyModuleNames.AddRange(
+					new string[] {
+					"LibWebM",
+				});
+
+				PublicDefinitions.Add("WITH_WEBM_LIBS=1");
+			} 
+			else if ((Target.IsInPlatformGroup(UnrealPlatformGroup.Unix) && Target.Architecture.StartsWith("x86_64")) || Target.Platform == UnrealTargetPlatform.Win64)
 			{
 				// libVPX is linked inside WebRTC so just use those headers and binaries where avaliable
 				PrivateIncludePaths.Add(Path.Join(EngineDirectory, "Source", "ThirdParty", WebRTC.WebRtcVersionPath, "Include", "third_party", "libvpx", "source", "libvpx"));
@@ -39,19 +52,6 @@ namespace UnrealBuildTool.Rules
 				
 				PublicDefinitions.Add("WITH_WEBM_LIBS=1");			
 			} 
-			else if (Target.Platform == UnrealTargetPlatform.Mac)
-			{
-				// Engine does not currently ship with Mac binaries for WebRTC so use fallback
-				
-				PrivateDependencyModuleNames.Add("LibVpx");
-				
-				PublicDependencyModuleNames.AddRange(
-					new string[] {
-					"LibWebM",
-				});
-
-				PublicDefinitions.Add("WITH_WEBM_LIBS=1");
-			} 
 			else
 			{
 				PublicDefinitions.Add("WITH_WEBM_LIBS=0");
