Make WebMMedia module rules compile when Stadia is not present.

diff --git a/Engine/Plugins/Media/WebMMedia/Source/WebMMedia/WebMMedia.Build.cs b/Engine/Plugins/Media/WebMMedia/Source/WebMMedia/WebMMedia.Build.cs
--- a/Engine/Plugins/Media/WebMMedia/Source/WebMMedia/WebMMedia.Build.cs
+++ b/Engine/Plugins/Media/WebMMedia/Source/WebMMedia/WebMMedia.Build.cs
@@ -25,7 +25,7 @@ namespace UnrealBuildTool.Rules
 					"Vorbis",
 				});
 
-			if ((Target.Platform == UnrealTargetPlatform.Stadia) || (Target.Platform == UnrealTargetPlatform.Mac))
+			if (Target.Platform == UnrealTargetPlatform.Mac)
 			{
 				// Engine does not currently ship with Mac or Stadia binaries for WebRTC so use fallback
 				
