Fix QAGame build on Linux

diff --git a/Engine/Plugins/Media/WebMMedia/Source/WebMMedia/WebMMedia.Build.cs b/Engine/Plugins/Media/WebMMedia/Source/WebMMedia/WebMMedia.Build.cs
--- a/Engine/Plugins/Media/WebMMedia/Source/WebMMedia/WebMMedia.Build.cs
+++ b/Engine/Plugins/Media/WebMMedia/Source/WebMMedia/WebMMedia.Build.cs
@@ -38,7 +38,7 @@ namespace UnrealBuildTool.Rules
 
 				PublicDefinitions.Add("WITH_WEBM_LIBS=1");
 			} 
-			else if ((Target.IsInPlatformGroup(UnrealPlatformGroup.Unix) && Target.Architecture.StartsWith("x86_64")) || Target.Platform == UnrealTargetPlatform.Win64)
+			else if (Target.Platform == UnrealTargetPlatform.Win64)
 			{
 				PublicDependencyModuleNames.AddRange(
 					new string[] {
@@ -47,7 +47,21 @@ namespace UnrealBuildTool.Rules
 				});
 				
 				PublicDefinitions.Add("WITH_WEBM_LIBS=1");			
-			} 
+			}
+			else if (Target.IsInPlatformGroup(UnrealPlatformGroup.Unix) && Target.Architecture.StartsWith("x86_64"))
+			{
+				// libVPX is linked inside WebRTC so just use those headers and binaries where avaliable
+				// In order to get rid of this need to rebuild webrtc without libvpx included (win64, macos, linux where editor runs)
+				PrivateIncludePaths.Add(Path.Join(EngineDirectory, "Source", "ThirdParty", WebRTC.WebRtcVersionPath, "Include", "third_party", "libvpx", "source", "libvpx"));
+
+				PublicDependencyModuleNames.AddRange(
+					new string[] {
+					"LibWebM",
+					"WebRTC"
+				});
+				
+				PublicDefinitions.Add("WITH_WEBM_LIBS=1");			
+			}
 			else
 			{
 				PublicDefinitions.Add("WITH_WEBM_LIBS=0");
