Fix WebMMedia runtime crash when trying to play media and cleanup WebMMedia.build.cs

diff --git a/Engine/Plugins/Media/WebMMedia/Source/WebMMedia/Private/Player/WebMVideoDecoder.cpp b/Engine/Plugins/Media/WebMMedia/Source/WebMMedia/Private/Player/WebMVideoDecoder.cpp
--- a/Engine/Plugins/Media/WebMMedia/Source/WebMMedia/Private/Player/WebMVideoDecoder.cpp
+++ b/Engine/Plugins/Media/WebMMedia/Source/WebMMedia/Private/Player/WebMVideoDecoder.cpp
@@ -4,11 +4,6 @@
 
 #if WITH_WEBM_LIBS
 
-THIRD_PARTY_INCLUDES_START
-#include <vpx/vpx_decoder.h>
-#include <vpx/vp8dx.h>
-THIRD_PARTY_INCLUDES_END
-
 #include "WebMMediaPrivate.h"
 #include "WebMMediaFrame.h"
 #include "WebMMediaTextureSample.h"
@@ -87,11 +82,11 @@ bool FWebMVideoDecoder::Initialize(const char* CodecName)
 	const vpx_codec_dec_cfg_t CodecConfig = { NumOfThreads, 0, 0 };
 	if (FCStringAnsi::Strcmp(CodecName, "V_VP8") == 0)
 	{
-		verify(vpx_codec_dec_init(Context, vpx_codec_vp8_dx(), &CodecConfig, /*VPX_CODEC_USE_FRAME_THREADING*/ 0) == 0);
+		verify(vpx_codec_dec_init(&Context, vpx_codec_vp8_dx(), &CodecConfig, /*VPX_CODEC_USE_FRAME_THREADING*/ 0) == 0);
 	}
 	else if (FCStringAnsi::Strcmp(CodecName, "V_VP9") == 0)
 	{
-		verify(vpx_codec_dec_init(Context, vpx_codec_vp9_dx(), &CodecConfig, /*VPX_CODEC_USE_FRAME_THREADING*/ 0) == 0);
+		verify(vpx_codec_dec_init(&Context, vpx_codec_vp9_dx(), &CodecConfig, /*VPX_CODEC_USE_FRAME_THREADING*/ 0) == 0);
 	}
 	else
 	{
@@ -114,10 +109,10 @@ void FWebMVideoDecoder::DecodeVideoFramesAsync(const TArray<TSharedPtr<FWebMFram
 		{
 			Prerequisites.Emplace(VideoDecodingTask);
 		}
-		VideoDecodingTask = FFunctionGraphTask::CreateAndDispatchWhenReady([this,  VideoFrames]()
-		{
-			DoDecodeVideoFrames(VideoFrames);
-		}, TStatId(), &Prerequisites, ENamedThreads::AnyThread);
+		VideoDecodingTask = FFunctionGraphTask::CreateAndDispatchWhenReady([this, VideoFrames]()
+			{
+				DoDecodeVideoFrames(VideoFrames);
+			}, TStatId(), &Prerequisites, ENamedThreads::AnyThread);
 	}
 	else
 	{
@@ -134,14 +129,14 @@ void FWebMVideoDecoder::DoDecodeVideoFrames(const TArray<TSharedPtr<FWebMFrame>>
 {
 	for (const TSharedPtr<FWebMFrame>& VideoFrame : VideoFrames)
 	{
-		if (vpx_codec_decode(Context, VideoFrame->Data.GetData(), VideoFrame->Data.Num(), nullptr, 0) != 0)
+		if (vpx_codec_decode(&Context, VideoFrame->Data.GetData(), VideoFrame->Data.Num(), nullptr, 0) != 0)
 		{
 			UE_LOG(LogWebMMedia, Display, TEXT("Error decoding video frame"));
 			return;
 		}
 
 		const void* ImageIter = nullptr;
-		while (const vpx_image_t* Image = vpx_codec_get_frame(Context, &ImageIter))
+		while (const vpx_image_t* Image = vpx_codec_get_frame(&Context, &ImageIter))
 		{
 			FWebMVideoDecoder* Self = this;
 			if (!bTexturesCreated)
@@ -209,7 +204,7 @@ void FWebMVideoDecoder::Close()
 
 	if (bIsInitialized)
 	{
-		vpx_codec_destroy(Context);
+		vpx_codec_destroy(&Context);
 		bIsInitialized = false;
 	}
 
diff --git a/Engine/Plugins/Media/WebMMedia/Source/WebMMedia/Public/WebMVideoDecoder.h b/Engine/Plugins/Media/WebMMedia/Source/WebMMedia/Public/WebMVideoDecoder.h
--- a/Engine/Plugins/Media/WebMMedia/Source/WebMMedia/Public/WebMVideoDecoder.h
+++ b/Engine/Plugins/Media/WebMMedia/Source/WebMMedia/Public/WebMVideoDecoder.h
@@ -7,12 +7,15 @@
 #include "Templates/SharedPointer.h"
 #include "MediaShaders.h"
 
+THIRD_PARTY_INCLUDES_START
+#include <vpx/vpx_decoder.h>
+#include <vpx/vp8dx.h>
+THIRD_PARTY_INCLUDES_END
+
 class IWebMSamplesSink;
 class FWebMMediaTextureSample;
 class FWebMMediaTextureSamplePool;
 struct FWebMFrame;
-struct vpx_image;
-struct vpx_codec_ctx;
 
 class WEBMMEDIA_API FWebMVideoDecoder
 {
@@ -29,10 +32,10 @@ private:
 	struct FConvertParams
 	{
 		TSharedPtr<FWebMMediaTextureSample, ESPMode::ThreadSafe> VideoSample;
-		const vpx_image* Image;
+		const vpx_image_t* Image;
 	};
 
-	vpx_codec_ctx* Context;
+	vpx_codec_ctx_t Context;
 	TUniquePtr<FWebMMediaTextureSamplePool> VideoSamplePool;
 	TRefCountPtr<FRHITexture2D> DecodedY;
 	TRefCountPtr<FRHITexture2D> DecodedU;
@@ -44,7 +47,7 @@ private:
 
 	void ConvertYUVToRGBAndSubmit(const FConvertParams& Params);
 	void DoDecodeVideoFrames(const TArray<TSharedPtr<FWebMFrame>>& VideoFrames);
-	void CreateTextures(const vpx_image* Image);
+	void CreateTextures(const vpx_image_t* Image);
 	void Close();
 };
 
diff --git a/Engine/Plugins/Media/WebMMedia/Source/WebMMedia/WebMMedia.Build.cs b/Engine/Plugins/Media/WebMMedia/Source/WebMMedia/WebMMedia.Build.cs
--- a/Engine/Plugins/Media/WebMMedia/Source/WebMMedia/WebMMedia.Build.cs
+++ b/Engine/Plugins/Media/WebMMedia/Source/WebMMedia/WebMMedia.Build.cs
@@ -25,11 +25,11 @@ namespace UnrealBuildTool.Rules
 					"Vorbis",
 				});
 
-			if (Target.Platform == UnrealTargetPlatform.Mac)
+			if (Target.Platform == UnrealTargetPlatform.Mac ||
+				Target.Platform == UnrealTargetPlatform.Win64 ||
+				(Target.IsInPlatformGroup(UnrealPlatformGroup.Unix) && Target.Architecture == UnrealArch.X64))
 			{
-				// Engine does not currently ship with Mac binaries for WebRTC so use fallback
-				
-				PrivateDependencyModuleNames.Add("LibVpx");
+				PublicDependencyModuleNames.Add("LibVpx");
 				
 				PublicDependencyModuleNames.AddRange(
 					new string[] {
@@ -37,30 +37,6 @@ namespace UnrealBuildTool.Rules
 				});
 
 				PublicDefinitions.Add("WITH_WEBM_LIBS=1");
-			} 
-			else if (Target.Platform == UnrealTargetPlatform.Win64)
-			{
-				PublicDependencyModuleNames.AddRange(
-					new string[] {
-					"LibWebM",
-					"LibVpx"
-				});
-				
-				PublicDefinitions.Add("WITH_WEBM_LIBS=1");			
-			}
-			else if (Target.IsInPlatformGroup(UnrealPlatformGroup.Unix) && Target.Architecture == UnrealArch.X64)
-			{
-				// libVPX is linked inside WebRTC so just use those headers and binaries where avaliable
-				// In order to get rid of this need to rebuild webrtc without libvpx included (win64, macos, linux where editor runs)
-				PrivateIncludePaths.Add(Path.Join(EngineDirectory, "Source", "ThirdParty", "WebRTC", "4664", "Include", "third_party", "libvpx", "source", "libvpx"));
-
-				PublicDependencyModuleNames.AddRange(
-					new string[] {
-					"LibWebM",
-					"WebRTC"
-				});
-				
-				PublicDefinitions.Add("WITH_WEBM_LIBS=1");			
 			}
 			else
 			{
