Added input parameterer checking

#rnx

diff --git a/Engine/Plugins/Media/WebMMedia/Source/WebMMedia/Private/Player/WebMMediaPlayer.cpp b/Engine/Plugins/Media/WebMMedia/Source/WebMMedia/Private/Player/WebMMediaPlayer.cpp
--- a/Engine/Plugins/Media/WebMMedia/Source/WebMMedia/Private/Player/WebMMediaPlayer.cpp
+++ b/Engine/Plugins/Media/WebMMedia/Source/WebMMedia/Private/Player/WebMMediaPlayer.cpp
@@ -299,7 +299,7 @@ void FWebMMediaPlayer::TickInput(FTimespan DeltaTime, FTimespan /*Timecode*/)
 
 bool FWebMMediaPlayer::GetAudioTrackFormat(int32 TrackIndex, int32 FormatIndex, FMediaAudioTrackFormat& OutFormat) const
 {
-	if (FormatIndex != 0)
+	if (TrackIndex < 0 || FormatIndex <= 0)
 	{
 		// We only support one format per track
 		return false;
