Changing Lock/Unlock code that uses GDynamicRHI directly to use their RHICommandLists instead.

diff --git a/Engine/Plugins/Media/BinkMedia/Source/BinkMediaPlayer/Private/binktexturesRHI.cpp b/Engine/Plugins/Media/BinkMedia/Source/BinkMediaPlayer/Private/binktexturesRHI.cpp
--- a/Engine/Plugins/Media/BinkMedia/Source/BinkMediaPlayer/Private/binktexturesRHI.cpp
+++ b/Engine/Plugins/Media/BinkMedia/Source/BinkMediaPlayer/Private/binktexturesRHI.cpp
@@ -293,7 +293,7 @@ static void Finish_texture_update( BINKTEXTURES * ptextures )
 static void update_plane_texture_rect(FRHICommandListImmediate& RHI, FRHITexture* RHITexture, BINKPLANE const* plane, unsigned w, unsigned h)
 {
 	uint32 Stride = 0;
-	unsigned char* TextureMemory = (unsigned char*)GDynamicRHI->LockTexture2D_RenderThread(RHI, RHITexture, 0, RLM_WriteOnly, Stride, false);
+	unsigned char* TextureMemory = (unsigned char*)RHI.LockTexture2D(RHITexture, 0, RLM_WriteOnly, Stride, false);
 
 	if (TextureMemory)
 	{
@@ -309,7 +309,7 @@ static void update_plane_texture_rect(FRHICommandListImmediate& RHI, FRHITexture
 				FMemory::Memcpy(TextureMemory + i * Stride, (char*)plane->Buffer + i * w, w);
 			}
 		}
-		GDynamicRHI->UnlockTexture2D_RenderThread(RHI, RHITexture, 0, false);
+		RHI.UnlockTexture2D(RHITexture, 0, false);
 	}
 }
 
diff --git a/Engine/Plugins/Media/WebMMedia/Source/WebMMedia/Private/Player/WebMVideoDecoder.cpp b/Engine/Plugins/Media/WebMMedia/Source/WebMMedia/Private/Player/WebMVideoDecoder.cpp
--- a/Engine/Plugins/Media/WebMMedia/Source/WebMMedia/Private/Player/WebMVideoDecoder.cpp
+++ b/Engine/Plugins/Media/WebMMedia/Source/WebMMedia/Private/Player/WebMVideoDecoder.cpp
@@ -228,7 +228,7 @@ void FWebMVideoDecoder::ConvertYUVToRGBAndSubmit(const FConvertParams& Params)
 			int CopyHeight)
 		{
 			uint32 Stride = 0;
-			unsigned char* TextureMemory = (unsigned char*)GDynamicRHI->LockTexture2D_RenderThread(InCommandList, RHITexture, 0, RLM_WriteOnly, Stride, false);
+			unsigned char* TextureMemory = (unsigned char*)InCommandList.LockTexture2D(RHITexture, 0, RLM_WriteOnly, Stride, false);
 
 			if (TextureMemory)
 			{
@@ -244,7 +244,7 @@ void FWebMVideoDecoder::ConvertYUVToRGBAndSubmit(const FConvertParams& Params)
 						FMemory::Memcpy(TextureMemory + h * Stride, Image->planes[ImageIndex] + h * Image->stride[ImageIndex], Image->stride[ImageIndex]);
 					}
 				}
-				GDynamicRHI->UnlockTexture2D_RenderThread(InCommandList, RHITexture, 0, false);
+				InCommandList.UnlockTexture2D(RHITexture, 0, false);
 			}
 		};
 
