WebMMedia: Clearing all members on Close() to remove any tracks that were added before detecting the media cannot be used, which would keep them around causing a crash on the next Open() call.

#rnx

diff --git a/Engine/Plugins/Media/WebMMedia/Source/WebMMedia/Private/Player/WebMMediaPlayer.cpp b/Engine/Plugins/Media/WebMMedia/Source/WebMMedia/Private/Player/WebMMediaPlayer.cpp
--- a/Engine/Plugins/Media/WebMMedia/Source/WebMMedia/Private/Player/WebMMediaPlayer.cpp
+++ b/Engine/Plugins/Media/WebMMedia/Source/WebMMedia/Private/Player/WebMMediaPlayer.cpp
@@ -38,11 +38,6 @@ FWebMMediaPlayer::~FWebMMediaPlayer()
 
 void FWebMMediaPlayer::Close()
 {
-	if (CurrentState == EMediaState::Closed)
-	{
-		return;
-	}
-
 	VideoDecoder.Reset();
 	AudioDecoder.Reset();
 	VideoTracks.Empty();
@@ -52,15 +47,19 @@ void FWebMMediaPlayer::Close()
 	MkvCurrentCluster = nullptr;
 	MkvCurrentBlockEntry = nullptr;
 	MediaUrl.Empty();
-	CurrentState = EMediaState::Closed;
 	SelectedAudioTrack = INDEX_NONE;
 	SelectedVideoTrack = INDEX_NONE;
 	CurrentTime = 0;
 	bNoMoreToRead = false;
+	bLooping = false;
 
-	// notify listeners
-	OutEvents.Push(EMediaEvent::TracksChanged);
-	OutEvents.Push(EMediaEvent::MediaClosed);
+	if (CurrentState != EMediaState::Closed)
+	{
+		CurrentState = EMediaState::Closed;
+		// notify listeners
+		OutEvents.Push(EMediaEvent::TracksChanged);
+		OutEvents.Push(EMediaEvent::MediaClosed);
+	}
 }
 
 IMediaCache& FWebMMediaPlayer::GetCache()
