WebMMedia: Rewrite of the plugin

#rnx

#Synced-CL 50037646

diff --git a/Engine/Plugins/Media/WebMMedia/Source/WebMMedia/Private/Player/WebMAudioDecoder.cpp b/Engine/Plugins/Media/WebMMedia/Source/WebMMedia/Private/Player/WebMAudioDecoder.cpp
--- a/Engine/Plugins/Media/WebMMedia/Source/WebMMedia/Private/Player/WebMAudioDecoder.cpp
+++ b/Engine/Plugins/Media/WebMMedia/Source/WebMMedia/Private/Player/WebMAudioDecoder.cpp
@@ -22,6 +22,7 @@ FWebMAudioDecoder::FWebMAudioDecoder(IWebMSamplesSink& InSamples)
 	: AudioSamplePool(new FWebMMediaAudioSamplePool)
 	, Samples(InSamples)
 	, OpusDecoder(nullptr)
+	, CurrentSequenceIndex(0)
 {
 }
 
@@ -40,42 +41,49 @@ FWebMAudioDecoder::~FWebMAudioDecoder()
 
 bool FWebMAudioDecoder::Initialize(const char* InCodec, int32 InSampleRate, int32 InChannels, const uint8* CodecPrivateData, size_t CodecPrivateDataSize)
 {
+	return InitializeInternal(FString(InCodec), InSampleRate, InChannels, MakeConstArrayView<uint8>(CodecPrivateData, (int32) CodecPrivateDataSize));
+}
+
+bool FWebMAudioDecoder::InitializeInternal(FString InCodec, int32 InSampleRate, int32 InChannels, TConstArrayView<uint8> InCodecPrivate)
+{
+	CodecSpecificData = InCodecPrivate;
+
 	if (OpusDecoder)
 	{
 		opus_decoder_destroy(OpusDecoder);
 		OpusDecoder = nullptr;
 	}
 
-	if (FCStringAnsi::Strcmp(InCodec, "A_OPUS") == 0)
+	if (InCodec.Equals(TEXT("A_OPUS")))
 	{
-		Codec = ESupportedCodecs::Opus;
+		Codec = ECodec::Opus;
 	}
-	else if (FCStringAnsi::Strcmp(InCodec, "A_VORBIS") == 0)
+	else if (InCodec.Equals(TEXT("A_VORBIS")))
 	{
-		Codec = ESupportedCodecs::Vorbis;
+		Codec = ECodec::Vorbis;
 	}
 	else
 	{
+		Codec = ECodec::Undefined;
 		return false;
 	}
 
 	SampleRate = InSampleRate;
 	Channels = InChannels;
 
-	if (Codec == ESupportedCodecs::Opus)
+	if (Codec == ECodec::Opus)
 	{
 		InitializeOpus();
 	}
 	else
 	{
-		if (!InitializeVorbis(CodecPrivateData, CodecPrivateDataSize))
+		if (!InitializeVorbis(CodecSpecificData.GetData(), CodecSpecificData.Num()))
 		{
 			return false;
 		}
 	}
 
 	DecodeBuffer.SetNumUninitialized(FrameSize * Channels * sizeof(int16));
-
 	return true;
 }
 
@@ -102,9 +110,9 @@ bool FWebMAudioDecoder::InitializeVorbis(const uint8* CodecPrivateData, size_t C
 	size_t Offset = 1;
 
 	// Calculate three headers sizes
-	for (int32 i = 0; i < 2; ++i)
+	for(int32 i = 0; i < 2; ++i)
 	{
-		for (;;)
+		for(;;)
 		{
 			if (Offset >= CodecPrivateDataSize)
 			{
@@ -146,7 +154,7 @@ bool FWebMAudioDecoder::InitializeVorbis(const uint8* CodecPrivateData, size_t C
 	/* Upload three Vorbis headers into libvorbis */
 	vorbis_comment Vc;
 	vorbis_comment_init(&Vc);
-	for (int i = 0; i < 3; ++i)
+	for(int i = 0; i < 3; ++i)
 	{
 		if (vorbis_synthesis_headerin(&VorbisDecoder->Info, &Vc, &Packet[i]))
 		{
@@ -181,7 +189,7 @@ bool FWebMAudioDecoder::InitializeVorbis(const uint8* CodecPrivateData, size_t C
 	return true;
 }
 
-void FWebMAudioDecoder::DecodeAudioFramesAsync(const TArray<TSharedPtr<FWebMFrame>>& AudioFrames)
+void FWebMAudioDecoder::DecodeAudioFramesAsync(const TArray<TSharedPtr<FWebMFrame, ESPMode::ThreadSafe>>& AudioFrames)
 {
 	FGraphEventRef PreviousDecodingTask = AudioDecodingTask;
 
@@ -201,13 +209,33 @@ bool FWebMAudioDecoder::IsBusy() const
 	return AudioDecodingTask && !AudioDecodingTask->IsComplete();
 }
 
-void FWebMAudioDecoder::DoDecodeAudioFrames(const TArray<TSharedPtr<FWebMFrame>>& AudioFrames)
+void FWebMAudioDecoder::DoDecodeAudioFrames(const TArray<TSharedPtr<FWebMFrame, ESPMode::ThreadSafe>>& AudioFrames)
 {
-	for (const TSharedPtr<FWebMFrame>& AudioFrame : AudioFrames)
+	for(const TSharedPtr<FWebMFrame, ESPMode::ThreadSafe>& AudioFrame : AudioFrames)
 	{
 		int32 NumOfSamplesDecoded;
 
-		if (Codec == ESupportedCodecs::Opus)
+		// Change in codec?
+		if (Codec == ECodec::Undefined)
+		{
+			if (AudioFrame->CodecName.IsEmpty())
+			{
+				UE_LOG(LogWebMMedia, Display, TEXT("No codec specified to initialize the codec lazily"));
+				Samples.ReportAudioDecodingError(TEXT("No codec specified to initialize the codec lazily"));
+				return;
+			}
+		}
+		if (AudioFrame->CodecName.Len())
+		{
+			if (!InitializeInternal(AudioFrame->CodecName, AudioFrame->SamplingRate, AudioFrame->NumChannels, AudioFrame->CodecSpecificData))
+			{
+				UE_LOG(LogWebMMedia, Display, TEXT("Failed to create codec lazily"));
+				Samples.ReportAudioDecodingError(TEXT("Failed to create codec lazily"));
+				return;
+			}
+		}
+
+		if (Codec == ECodec::Opus)
 		{
 			NumOfSamplesDecoded = DecodeOpus(AudioFrame);
 		}
@@ -220,15 +248,27 @@ void FWebMAudioDecoder::DoDecodeAudioFrames(const TArray<TSharedPtr<FWebMFrame>>
 		{
 			TSharedRef<FWebMMediaAudioSample, ESPMode::ThreadSafe> AudioSample = AudioSamplePool->AcquireShared();
 			FTimespan Duration = (NumOfSamplesDecoded * ETimespan::TicksPerSecond) / SampleRate;
-			AudioSample->Initialize(DecodeBuffer.GetData(), NumOfSamplesDecoded * Channels * 2, Channels, SampleRate, AudioFrame->Time, Duration);
+			AudioSample->Initialize(DecodeBuffer.GetData(), NumOfSamplesDecoded * Channels * 2, Channels, SampleRate, FMediaTimeStamp(AudioFrame->Time, AudioFrame->SequenceIndex), Duration, AudioFrame->DecoderIndex);
 
 			Samples.AddAudioSampleFromDecodingThread(AudioSample);
 		}
+		else if (NumOfSamplesDecoded < 0)
+		{
+			Samples.ReportAudioDecodingError(TEXT("Error decoding audio frame"));
+		}
 	}
 }
 
-int32 FWebMAudioDecoder::DecodeOpus(const TSharedPtr<FWebMFrame>& AudioFrame)
+int32 FWebMAudioDecoder::DecodeOpus(const TSharedPtr<FWebMFrame, ESPMode::ThreadSafe>& AudioFrame)
 {
+	// Change in sequence?
+	if (AudioFrame->SequenceIndex != CurrentSequenceIndex)
+	{
+		opus_decoder_destroy(OpusDecoder);
+		InitializeOpus();
+		CurrentSequenceIndex = AudioFrame->SequenceIndex;
+	}
+
 	int32 Result = opus_decode(OpusDecoder, AudioFrame->Data.GetData(), AudioFrame->Data.Num(), (opus_int16*) DecodeBuffer.GetData(), FrameSize, 0);
 	if (Result < 0)
 	{
@@ -238,8 +278,15 @@ int32 FWebMAudioDecoder::DecodeOpus(const TSharedPtr<FWebMFrame>& AudioFrame)
 	return Result;
 }
 
-int32 FWebMAudioDecoder::DecodeVorbis(const TSharedPtr<FWebMFrame>& AudioFrame)
+int32 FWebMAudioDecoder::DecodeVorbis(const TSharedPtr<FWebMFrame, ESPMode::ThreadSafe>& AudioFrame)
 {
+	// Change in sequence?
+	if (AudioFrame->SequenceIndex != CurrentSequenceIndex)
+	{
+		InitializeVorbis(CodecSpecificData.GetData(), CodecSpecificData.Num());
+		CurrentSequenceIndex = AudioFrame->SequenceIndex;
+	}
+
 	ogg_packet Packet;
 	memset(&Packet, 0, sizeof(ogg_packet));
 
@@ -255,7 +302,7 @@ int32 FWebMAudioDecoder::DecodeVorbis(const TSharedPtr<FWebMFrame>& AudioFrame)
 	if (vorbis_synthesis_blockin(&VorbisDecoder->DspState, &VorbisDecoder->Block))
 	{
 		UE_LOG(LogWebMMedia, Warning, TEXT("Error decoding Vorbis audio frame - vorbis_synthesis_blockin failed"));
-		return - 1;
+		return -1;
 	}
 
 	int16* Buffer = (int16*) DecodeBuffer.GetData();
@@ -265,12 +312,12 @@ int32 FWebMAudioDecoder::DecodeVorbis(const TSharedPtr<FWebMFrame>& AudioFrame)
 	while (SamplesCount > 0)
 	{
 		int32 ToConvert = SamplesCount <= FrameSize ? SamplesCount : FrameSize;
-		for (int32 Channel = 0; Channel < Channels; ++Channel)
+		for(int32 Channel = 0; Channel < Channels; ++Channel)
 		{
 			float *FloatSamples = Pcm[Channel];
-			for (int32 i = 0, j = Channel; i < ToConvert; ++i, j += Channels)
+			for(int32 i = 0, j = Channel; i < ToConvert; ++i, j += Channels)
 			{
-				int32 Sample = FMath::Clamp<int32>(FloatSamples[i] * 32767.0f, -32768, 32768);
+				int32 Sample = FMath::Clamp<int32>(FloatSamples[i] * 32767.0f, -32768, 32767);
 				Buffer[Count + j] = Sample;
 			}
 		}
diff --git a/Engine/Plugins/Media/WebMMedia/Source/WebMMedia/Private/Player/WebMMediaPlayer.cpp b/Engine/Plugins/Media/WebMMedia/Source/WebMMedia/Private/Player/WebMMediaPlayer.cpp
--- a/Engine/Plugins/Media/WebMMedia/Source/WebMMedia/Private/Player/WebMMediaPlayer.cpp
+++ b/Engine/Plugins/Media/WebMMedia/Source/WebMMedia/Private/Player/WebMMediaPlayer.cpp
@@ -20,14 +20,6 @@
 FWebMMediaPlayer::FWebMMediaPlayer(IMediaEventSink& InEventSink)
 	: EventSink(InEventSink)
 	, Samples(MakeShared<FMediaSamples, ESPMode::ThreadSafe>())
-	, MkvCurrentCluster(nullptr)
-	, MkvCurrentBlockEntry(nullptr)
-	, CurrentState(EMediaState::Closed)
-	, SelectedAudioTrack(INDEX_NONE)
-	, SelectedVideoTrack(INDEX_NONE)
-	, CurrentTime(0)
-	, bNoMoreToRead(false)
-	, bLooping(false)
 {
 }
 
@@ -36,6 +28,54 @@ FWebMMediaPlayer::~FWebMMediaPlayer()
 	Close();
 }
 
+FWebMMediaPlayer::ETrackCodec FWebMMediaPlayer::GetCodecFromMKVTrack(const char* InMKVCodecName)
+{
+	if (FCStringAnsi::Strcmp(InMKVCodecName, "V_VP8") == 0)
+	{
+		return ETrackCodec::VP8;
+	}
+	else if (FCStringAnsi::Strcmp(InMKVCodecName, "V_VP9") == 0)
+	{
+		return ETrackCodec::VP9;
+	}
+	else if (FCStringAnsi::Strcmp(InMKVCodecName, "A_OPUS") == 0)
+	{
+		return ETrackCodec::Opus;
+	}
+	else if (FCStringAnsi::Strcmp(InMKVCodecName, "A_VORBIS") == 0)
+	{
+		return ETrackCodec::Vorbis;
+	}
+	return ETrackCodec::Undefined;
+}
+FString FWebMMediaPlayer::GetMKVCodec(ETrackCodec InCodec)
+{
+	switch(InCodec)
+	{
+		case ETrackCodec::VP8:
+		{
+			return FString(TEXT("V_VP8"));
+		}
+		case ETrackCodec::VP9:
+		{
+			return FString(TEXT("V_VP9"));
+		}
+		case ETrackCodec::Opus:
+		{
+			return FString(TEXT("A_OPUS"));
+		}
+		case ETrackCodec::Vorbis:
+		{
+			return FString(TEXT("A_VORBIS"));
+		}
+		default:
+		{
+			return FString(TEXT("X_UNKNOWN"));
+		}
+	}
+}
+
+
 void FWebMMediaPlayer::Close()
 {
 	VideoDecoder.Reset();
@@ -49,10 +89,47 @@ void FWebMMediaPlayer::Close()
 	MediaUrl.Empty();
 	SelectedAudioTrack = INDEX_NONE;
 	SelectedVideoTrack = INDEX_NONE;
-	CurrentTime = 0;
-	bNoMoreToRead = false;
+	NewPendingVideoCodec.Reset();
+	NewPendingAudioCodec.Reset();
+	bPendingVideoTrackChange = false;
+	bPendingAudioTrackChange = false;
+	CurrentSeekIndex = 0;
+	CurrentLoopIndex = 0;
+	CurrentTime = FTimespan::Zero();
+	SeekToTime.Reset();
+	SeekTo.Reset();
+	CurrentVideoTime.Reset();
+	CurrentAudioTime.Reset();
+	bFileReachedEOS = false;
+	bFileError = false;
+	bDecodeError = false;
 	bLooping = false;
 
+	VideoFramesReadyForDecode.Empty();
+	AudioFramesReadyForDecode.Empty();
+	LoadedVideoFrames.Empty();
+	LoadedAudioFrames.Empty();
+	MovieDuration = FTimespan::Zero();
+	VideoFrameRate.Numerator = 0;
+	VideoFrameRate.Denominator = 0;
+	bRecalculateFrameRate = false;
+
+	VideoFrameDoneDelegate.Unbind();
+	AudioFrameDoneDelegate.Unbind();
+
+	DecodedSamplesLock.Lock();
+	NumDecodingEnqueuedVideoFrames = 0;
+	NumDecodingEnqueuedAudioFrames = 0;
+	DecodedVideoFrames.Empty();
+	DecodedAudioFrames.Empty();
+	DecodedSamplesLock.Unlock();
+
+	TSharedPtr<FMediaSamples, ESPMode::ThreadSafe> NewSamples(MakeShared<FMediaSamples, ESPMode::ThreadSafe>());
+	MediaSamplesLock.Lock();
+	Swap(Samples, NewSamples);
+	MediaSamplesLock.Unlock();
+	NewSamples.Reset();
+
 	if (CurrentState != EMediaState::Closed)
 	{
 		CurrentState = EMediaState::Closed;
@@ -85,6 +162,7 @@ FGuid FWebMMediaPlayer::GetPlayerPluginGUID() const
 
 IMediaSamples& FWebMMediaPlayer::GetSamples()
 {
+	FScopeLock lock(&MediaSamplesLock);
 	return *Samples.Get();
 }
 
@@ -151,146 +229,479 @@ bool FWebMMediaPlayer::Open(const FString& Url, const IMediaOptions* /*Options*/
 
 	VideoDecoder.Reset(new FWebMVideoDecoder(*this));
 	AudioDecoder.Reset(new FWebMAudioDecoder(*this));
-
-	if (SelectedAudioTrack != INDEX_NONE)
-	{
-		const mkvparser::AudioTrack* AudioTrack = AudioTracks[SelectedAudioTrack];
-
-		size_t PrivateDataSize;
-		const uint8* PrivateData = AudioTrack->GetCodecPrivate(PrivateDataSize);
-
-		AudioDecoder->Initialize(AudioTrack->GetCodecId(), AudioTrack->GetSamplingRate(), AudioTrack->GetChannels(), PrivateData, PrivateDataSize);
-	}
-
-	if (SelectedVideoTrack != INDEX_NONE)
-	{
-		const mkvparser::VideoTrack* VideoTrack = VideoTracks[SelectedVideoTrack];
-
-		VideoDecoder->Initialize(VideoTrack->GetCodecId());
-	}
-
 	CurrentState = EMediaState::Stopped;
-	bNoMoreToRead = false;
+	bFileReachedEOS = false;
+	bFileError = false;
+	bDecodeError = false;
 
 	// notify listeners
 	OutEvents.Push(EMediaEvent::TracksChanged);
 	OutEvents.Push(EMediaEvent::MediaOpened);
 
+	VideoFrameDoneDelegate.BindRaw(this, &FWebMMediaPlayer::ReturnVideoSample);
+	AudioFrameDoneDelegate.BindRaw(this, &FWebMMediaPlayer::ReturnAudioSample);
+
 	return true;
 }
 
 bool FWebMMediaPlayer::Open(const TSharedRef<FArchive, ESPMode::ThreadSafe>& /*Archive*/, const FString& /*OriginalUrl*/, const IMediaOptions* /*Options*/)
 {
-	// TODO! We don't support opening URLs for now.
+	// We do not support opening archives for now.
 	return false;
 }
 
-void FWebMMediaPlayer::TickFetch(FTimespan DeltaTime, FTimespan Timecode)
+void FWebMMediaPlayer::TickInternal()
 {
-	if (CurrentState != EMediaState::Playing)
+	// Load future data unless in a state where it is not possible or does not make sense.
+	if (CurrentState == EMediaState::Closed ||
+		CurrentState == EMediaState::Error ||
+		CurrentState == EMediaState::Stopped ||
+		bFileError || bDecodeError)
 	{
 		return;
 	}
 
-	CurrentTime += DeltaTime;
+	// Is a seek pending?
+	bool bSeekPending = SeekToTime.IsSet();
+	if (bSeekPending)
+	{
+		CurrentTime = SeekToTime.GetValue();
+		EventSink.ReceiveMediaEvent(EMediaEvent::SeekCompleted);
+	}
 
-	if (bNoMoreToRead)
+	// Is there a pending seek or track change?
+	if (bSeekPending || bPendingVideoTrackChange || bPendingAudioTrackChange)
 	{
-		// Check if we finished playing
-		if (Samples->NumVideoSamples() == 0 && Samples->NumAudioSamples() == 0)
+		// Clamp time if necessary
+		if (CurrentTime < FTimespan::Zero())
+		{
+			CurrentTime = FTimespan::Zero();
+		}
+		else if (CurrentTime > MovieDuration)
+		{
+			CurrentTime = MovieDuration;
+		}
+
+		bool bDeselectVideo = bPendingVideoTrackChange && SelectedVideoTrack == INDEX_NONE;
+		bool bDeselectAudio = bPendingAudioTrackChange && SelectedAudioTrack == INDEX_NONE;
+		bool bDeselectOnly = !bSeekPending &&
+							 ((bPendingVideoTrackChange && bDeselectVideo && !bPendingAudioTrackChange) ||
+							  (bPendingAudioTrackChange && bDeselectAudio && !bPendingVideoTrackChange) ||
+							  (bDeselectVideo && bDeselectAudio));
+
+		if (!bDeselectOnly)
 		{
-			if (bLooping)
+			// A track change is similar to a seek in that we have to restart at a keyframe as well.
+			SeekTo.Reset();
+			SeekTo.TargetTime = CurrentTime;
+			SeekTo.bAwaitTime = true;
+			SeekToTime.Reset();
+			CurrentVideoTime.Reset();
+			CurrentAudioTime.Reset();
+
+			// We need to fetch data for different tracks or at a different time now, so dump everything we have going at the moment.
+			VideoFramesReadyForDecode.Empty();
+			AudioFramesReadyForDecode.Empty();
+			LoadedVideoFrames.Empty();
+			LoadedAudioFrames.Empty();
+			Samples->FlushSamples();
+			if (bPendingVideoTrackChange && SelectedVideoTrack != INDEX_NONE)
 			{
-				Seek(FTimespan::FromSeconds(0));
+				NewPendingVideoCodec = GetCodecFromMKVTrack(VideoTracks[SelectedVideoTrack]->GetCodecId());
 			}
-			else
+			if (bPendingAudioTrackChange && SelectedAudioTrack != INDEX_NONE)
 			{
-				CurrentState = EMediaState::Stopped;
-				CurrentTime = FTimespan::FromMicroseconds((double) MkvSegment->GetDuration() / 1000.0);
-				OutEvents.Push(EMediaEvent::PlaybackEndReached);
-				OutEvents.Push(EMediaEvent::PlaybackSuspended);
+				NewPendingAudioCodec = GetCodecFromMKVTrack(AudioTracks[SelectedAudioTrack]->GetCodecId());
 			}
+			MkvSeekToTime(CurrentTime);
+			bFileReachedEOS = false;
+			++CurrentDecoderIndex;
 		}
+		else
+		{
+			if (bDeselectVideo)
+			{
+				VideoFramesReadyForDecode.Empty();
+				LoadedVideoFrames.Empty();
+				/*
+					TSharedPtr<IMediaTextureSample, ESPMode::ThreadSafe> VideoSample;
+					TRange<FTimespan> TimeRange(FTimespan::Zero(), FTimespan::MaxValue());
+					Samples->FetchVideo(TimeRange, VideoSample);
+				*/
+			}
+			if (bDeselectAudio)
+			{
+				AudioFramesReadyForDecode.Empty();
+				LoadedAudioFrames.Empty();
+				/*
+					TSharedPtr<IMediaAudioSample, ESPMode::ThreadSafe> AudioSample;
+					TRange<FTimespan> TimeRange(FTimespan::Zero(), FTimespan::MaxValue());
+					Samples->FetchAudio(TimeRange, AudioSample);
+				*/
+			}
+		}
+		bPendingVideoTrackChange = false;
+		bPendingAudioTrackChange = false;
+	}
 
+	// If no track is selected we do not need to load anything
+	if (SelectedVideoTrack == INDEX_NONE && SelectedAudioTrack == INDEX_NONE)
+	{
 		return;
 	}
 
-	FTimespan ReadBufferLength = FTimespan::FromSeconds(1.0);
-	FTimespan CurrentReadTime = CurrentTime;
-
-	TArray<TSharedPtr<FWebMFrame>> VideoFrames;
-	TArray<TSharedPtr<FWebMFrame>> AudioFrames;
-
-	if (CurrentReadTime < CurrentTime - ReadBufferLength / 2)
+	if (bFileReachedEOS)
 	{
-		UE_LOG(LogWebMMedia, Warning, TEXT("Playback is very behind, try to catchup and reset sync"));
-		Samples->FlushSamples();
+		if (bLooping)
+		{
+			MkvCurrentCluster = MkvSegment->FindCluster(0);
+			MkvCurrentBlockEntry = nullptr;
+			++CurrentLoopIndex;
+			bFileReachedEOS = false;
+		}
+		else
+		{
+			return;
+		}
 	}
 
+	// Number of future video frames we want to have, regardless of framerate.
+	int32 kNumFutureVideoFrames = VideoFrameRate.IsValid() ? (int32) VideoFrameRate.AsDecimal() : 30;
+	kNumFutureVideoFrames = kNumFutureVideoFrames > 60 ? 60 : kNumFutureVideoFrames < 5 ? 5 : kNumFutureVideoFrames;
+	const int32 kNumFutureAudioFrames = 60;
+
+	int64 CurrentSequenceIndex = FMediaTimeStamp::MakeIndexValue(CurrentSeekIndex, CurrentLoopIndex);
+
 	// Read frames up to 1 secs in the future
-	while (CurrentReadTime < CurrentTime + ReadBufferLength)
+	while(1)
 	{
+		bool bHaveEnough = true;
+		const bool bNeedVideo = SelectedVideoTrack != INDEX_NONE;
+		const bool bNeedAudio = SelectedAudioTrack != INDEX_NONE;
+		if (bNeedVideo)
+		{
+			bHaveEnough = (VideoFramesReadyForDecode.Num() + LoadedVideoFrames.Num()) >= kNumFutureVideoFrames;
+		}
+		else if (bNeedAudio)
+		{
+			bHaveEnough = (AudioFramesReadyForDecode.Num() + LoadedAudioFrames.Num()) >= kNumFutureAudioFrames;
+		}
+		if (bHaveEnough)
+		{
+			break;
+		}
+
 		MkvSeekToNextValidBlock();
 
-		if (bNoMoreToRead)
+		if (bFileReachedEOS || bFileError)
 		{
 			break;
 		}
 
-		CurrentReadTime = FTimespan::FromMicroseconds(MkvCurrentBlockEntry->GetBlock()->GetTime(MkvCurrentCluster) / 1000.0);
+		auto Blk = MkvCurrentBlockEntry->GetBlock();
+		const auto TrackNumber = Blk->GetTrackNumber();
+		const bool bIsVideo = bNeedVideo && VideoTracks[SelectedVideoTrack]->GetNumber() == TrackNumber;
+		const bool bIsAudio = bNeedAudio && AudioTracks[SelectedAudioTrack]->GetNumber() == TrackNumber;
+		if (!bIsVideo && !bIsAudio)
+		{
+			continue;
+		}
 
-		for (int32 i = 0; i < MkvCurrentBlockEntry->GetBlock()->GetFrameCount(); ++i)
+		FTimespan Timestamp(MkvCurrentBlockEntry->GetBlock()->GetTime(MkvCurrentCluster) / 100);
+		const bool bIsKeyframe = MkvCurrentBlockEntry->GetBlock()->IsKey();
+		const bool bIsInvisible = MkvCurrentBlockEntry->GetBlock()->IsInvisible();
+		check(!bIsInvisible);
+
+		// Wait for target time to appear?
+		if (SeekTo.bAwaitTime)
 		{
-			const mkvparser::Block::Frame& MkvFrame = MkvCurrentBlockEntry->GetBlock()->GetFrame(i);
+			// We need a video keyframe
+			if (bIsVideo)
+			{
+				if (bIsKeyframe)
+				{
+					SeekTo.bGotVideoKeyframe = true;
+					SeekTo.TargetTime = Timestamp;
+					CurrentTime = Timestamp;
+				}
+				else if (!SeekTo.bGotVideoKeyframe)
+				{
+					continue;
+				}
+			}
+			// If this is an audio sample and we need video but haven't found a video keyframe yet
+			// then this audio sample is not one we want.
+			else if (bIsAudio && bNeedVideo && !SeekTo.bGotVideoKeyframe)
+			{
+				continue;
+			}
+			// Anything that is before the time we want we do not need.
+			if (Timestamp < SeekTo.TargetTime)
+			{
+				continue;
+			}
+			else
+			{
+				SeekTo.Reset();
+			}
+		}
 
-			TSharedPtr<FWebMFrame> Frame = MakeShared<FWebMFrame>();
-			Frame->Time = CurrentReadTime;
-			Frame->Data.SetNumUninitialized(MkvFrame.len);
-			MkvFrame.Read(MkvReader.Get(), Frame->Data.GetData());
+		const int32 FrameCount = (int32) MkvCurrentBlockEntry->GetBlock()->GetFrameCount();
+		for(int32 i=0; i<FrameCount; ++i)
+		{
+			const mkvparser::Block::Frame& MkvFrame = MkvCurrentBlockEntry->GetBlock()->GetFrame(i);
 
-			int32 TrackNumber = MkvCurrentBlockEntry->GetBlock()->GetTrackNumber();
+			TSharedPtr<FWebMFrame, ESPMode::ThreadSafe> Frame = MakeShared<FWebMFrame>();
+			Frame->Time = i==0 ? Timestamp : FTimespan::MinValue();
+			Frame->SequenceIndex = CurrentSequenceIndex;
+			Frame->DecoderIndex = CurrentDecoderIndex;
+			Frame->bIsKeyframe = bIsKeyframe;
 
-			if (SelectedVideoTrack != INDEX_NONE && VideoTracks[SelectedVideoTrack]->GetNumber() == TrackNumber)
+			Frame->Data.SetNumUninitialized(MkvFrame.len);
+			bool bReadOk = 0 == MkvFrame.Read(MkvReader.Get(), Frame->Data.GetData());
+			if (bIsVideo)
 			{
-				Frame->Duration = FMkvFileReader::GetVideoFrameDuration(*VideoTracks[SelectedVideoTrack]);
-				VideoFrames.Add(Frame);
-				CurrentReadTime += Frame->Duration;
+				LoadedVideoFrames.Add(Frame);
 			}
-			else if (SelectedAudioTrack != INDEX_NONE && AudioTracks[SelectedAudioTrack]->GetNumber() == TrackNumber)
+			else if (bIsAudio)
 			{
-				// We will set duration after decompression.
-				AudioFrames.Add(Frame);
+				LoadedAudioFrames.Add(Frame);
 			}
 		}
+	}
 
-		// Make sure we're not decoding too much at once as it is GPU memory intensive.
-		if (VideoFrames.Num() > 20)
+	auto FixDurationsAndTimestamps = [](TArray<TSharedPtr<FWebMFrame, ESPMode::ThreadSafe>>& InFrames, FTimespan& InOutDurationSum, bool bIsFinalFrame, FTimespan InMovieDuration) -> int32
+	{
+		if (bIsFinalFrame)
 		{
-			break;
+			TSharedPtr<FWebMFrame, ESPMode::ThreadSafe> Frame = MakeShared<FWebMFrame>();
+			Frame->Time = InMovieDuration;
+			InFrames.Emplace(MoveTemp(Frame));
 		}
-	}
 
-	// Trigger video decoding
-	if (VideoFrames.Num() > 0)
+		int32 LastGoodSampleIndex = -1;
+		for(int32 i=0, iMax=InFrames.Num(); i<iMax; ++i)
+		{
+			if (InFrames[i]->Duration == FTimespan::Zero())
+			{
+				// Look forward to where the next sample that has a timestamp is.
+				for(int32 j=i+1; j<iMax; ++j)
+				{
+					if (InFrames[j]->Time > FTimespan::Zero())
+					{
+						int32 NumF = j - i;
+						FTimespan DurationSpan = InFrames[j]->Time - InFrames[i]->Time;
+						if (DurationSpan < FTimespan::Zero())
+						{
+							DurationSpan = FTimespan::Zero();
+						}
+						InOutDurationSum += DurationSpan;
+						if (NumF == 1)
+						{
+							InFrames[i]->Duration = DurationSpan;
+							LastGoodSampleIndex = i;
+						}
+						else
+						{
+							FTimespan DurPerSample = DurationSpan / NumF;
+							for(int32 k=i; k<j-1; ++k)
+							{
+								InFrames[k]->Duration = DurPerSample;
+								InFrames[k+1]->Time = InFrames[k]->Time + InFrames[k]->Duration;
+							}
+							InFrames[j-1]->Duration = InFrames[j]->Time - InFrames[j-1]->Time;
+							i = j - 1;
+							LastGoodSampleIndex = i;
+						}
+						break;
+					}
+				}
+			}
+		}
+		if (bIsFinalFrame)
+		{
+			InFrames.Pop();
+			LastGoodSampleIndex = InFrames.Num() - 1;
+		}
+		return LastGoodSampleIndex;
+	};
+
+
+	FTimespan DurationSumA, DurationSumV;
+	int32 LastGoodVideoSampleIndex = FixDurationsAndTimestamps(LoadedVideoFrames, DurationSumV, bFileReachedEOS, MovieDuration);
+	int32 LastGoodAudioSampleIndex = FixDurationsAndTimestamps(LoadedAudioFrames, DurationSumA, bFileReachedEOS, MovieDuration);
+	if (bRecalculateFrameRate && LastGoodVideoSampleIndex > 1)
 	{
-		VideoDecoder->DecodeVideoFramesAsync(VideoFrames);
+		VideoFrameRate.Numerator = ETimespan::TicksPerSecond;
+		VideoFrameRate.Denominator = DurationSumV.GetTicks() / (LastGoodVideoSampleIndex + 1);
+		bRecalculateFrameRate = false;
 	}
 
-	// Trigger audio decoding
-	if (AudioFrames.Num() > 0)
+	// Move the valid samples from the load to the decode queue.
+	if (LastGoodVideoSampleIndex >= 0)
 	{
-		AudioDecoder->DecodeAudioFramesAsync(AudioFrames);
+		for(int32 i=0; i<=LastGoodVideoSampleIndex; ++i)
+		{
+			VideoFramesReadyForDecode.Emplace(LoadedVideoFrames[0]);
+			LoadedVideoFrames.RemoveAt(0);
+		}
+	}
+	if (LastGoodAudioSampleIndex >= 0)
+	{
+		for(int32 i=0; i<=LastGoodAudioSampleIndex; ++i)
+		{
+			AudioFramesReadyForDecode.Emplace(LoadedAudioFrames[0]);
+			LoadedAudioFrames.RemoveAt(0);
+		}
 	}
 }
 
 void FWebMMediaPlayer::TickInput(FTimespan DeltaTime, FTimespan /*Timecode*/)
 {
+	TickInternal();
+
 	// Send out events
-	for (const auto& Event : OutEvents)
+	for(const auto& Event : OutEvents)
 	{
 		EventSink.ReceiveMediaEvent(Event);
 	}
 	OutEvents.Empty();
+	if (CurrentState == EMediaState::Error)
+	{
+		return;
+	}
+
+
+	// Enqueue frames for decoding.
+	if (!bDecodeError)
+	{
+		const int32 kMaxPendingVideoFrames = 5;
+		const int32 kMaxPendingAudioFrames = 10;
+
+		if (NumDecodingEnqueuedVideoFrames < kMaxPendingVideoFrames)
+		{
+			int32 FramesNeeded = kMaxPendingVideoFrames - NumDecodingEnqueuedVideoFrames;
+			TArray<TSharedPtr<FWebMFrame, ESPMode::ThreadSafe>> FramesForDecode;
+			FScopeLock lock(&DecodedSamplesLock);
+			for(int32 i=0; i<FramesNeeded && VideoFramesReadyForDecode.Num(); ++i)
+			{
+				FramesForDecode.Emplace(VideoFramesReadyForDecode[0]);	//-V767
+				VideoFramesReadyForDecode.RemoveAt(0);
+			}
+			if (FramesForDecode.Num())
+			{
+				// Is there a change in codec?
+				if (NewPendingVideoCodec.IsSet())
+				{
+					FramesForDecode[0]->CodecName = GetMKVCodec(NewPendingVideoCodec.GetValue());
+					NewPendingVideoCodec.Reset();
+				}
+				NumDecodingEnqueuedVideoFrames += FramesForDecode.Num();
+				VideoDecoder->DecodeVideoFramesAsync(FramesForDecode);
+			}
+		}
+
+		if (NumDecodingEnqueuedAudioFrames < kMaxPendingAudioFrames)
+		{
+			int32 FramesNeeded = kMaxPendingAudioFrames - NumDecodingEnqueuedAudioFrames;
+			TArray<TSharedPtr<FWebMFrame, ESPMode::ThreadSafe>> FramesForDecode;
+			FScopeLock lock(&DecodedSamplesLock);
+			for(int32 i=0; i<FramesNeeded && AudioFramesReadyForDecode.Num(); ++i)
+			{
+				FramesForDecode.Emplace(AudioFramesReadyForDecode[0]);	//-V767
+				AudioFramesReadyForDecode.RemoveAt(0);
+			}
+			if (FramesForDecode.Num())
+			{
+				// Is there a change in codec?
+				if (NewPendingAudioCodec.IsSet())
+				{
+					FramesForDecode[0]->CodecName = GetMKVCodec(NewPendingAudioCodec.GetValue());
+					NewPendingAudioCodec.Reset();
+					size_t CSDSize = 0;
+					const uint8* CSD = AudioTracks[SelectedAudioTrack]->GetCodecPrivate(CSDSize);
+					FramesForDecode[0]->CodecSpecificData = MakeConstArrayView<uint8>(CSD, (int32)CSDSize);
+					FramesForDecode[0]->NumChannels = (int32) AudioTracks[SelectedAudioTrack]->GetChannels();
+					FramesForDecode[0]->SamplingRate = (int32) AudioTracks[SelectedAudioTrack]->GetSamplingRate();
+				}
+				NumDecodingEnqueuedAudioFrames += FramesForDecode.Num();
+				AudioDecoder->DecodeAudioFramesAsync(FramesForDecode);
+			}
+		}
+	}
+
+	// Pass decoded samples to the media sample queue.
+	DecodedSamplesLock.Lock();
+	while(Samples->CanReceiveVideoSamples(1) && DecodedVideoFrames.Num())
+	{
+		auto VideoFrame = DecodedVideoFrames[0];
+		DecodedVideoFrames.RemoveAt(0);
+		Samples->AddVideo(VideoFrame.ToSharedRef());
+		--NumDecodingEnqueuedVideoFrames;
+	}
+	while(Samples->CanReceiveAudioSamples(1) && DecodedAudioFrames.Num())
+	{
+		auto AudioFrame = DecodedAudioFrames[0];
+		DecodedAudioFrames.RemoveAt(0);
+		Samples->AddAudio(AudioFrame.ToSharedRef());
+		--NumDecodingEnqueuedAudioFrames;
+	}
+	DecodedSamplesLock.Unlock();
+
+	// Encountered an error?
+	if ((bFileError || bDecodeError) && CurrentState != EMediaState::Error)
+	{
+		CurrentState = EMediaState::Error;
+		OutEvents.Push(EMediaEvent::PlaybackSuspended);
+		LoadedVideoFrames.Empty();
+		LoadedAudioFrames.Empty();
+		VideoFramesReadyForDecode.Empty();
+		AudioFramesReadyForDecode.Empty();
+		return;
+	}
+
+	// Update current time with either the video or audio time
+	if (CurrentState == EMediaState::Playing)
+	{
+		if (SelectedVideoTrack != INDEX_NONE && VideoTracks.Num())
+		{
+			FScopeLock lock(&TimeLock);
+			if (CurrentVideoTime.IsSet())
+			{
+				CurrentTime = CurrentVideoTime.GetValue();
+			}
+		}
+		else if (SelectedAudioTrack != INDEX_NONE && AudioTracks.Num())
+		{
+			FScopeLock lock(&TimeLock);
+			if (CurrentAudioTime.IsSet())
+			{
+				CurrentTime = CurrentAudioTime.GetValue();
+			}
+		}
+		else
+		{
+			CurrentTime += DeltaTime;
+			if (CurrentTime > MovieDuration)
+			{
+				CurrentTime = MovieDuration;
+			}
+		}
+	}
+
+	// Reached the end of the file while reading?
+	if (!bLooping && bFileReachedEOS && CurrentState == EMediaState::Playing)
+	{
+		// Check if we have processed all pending input.
+		if (VideoFramesReadyForDecode.IsEmpty() && AudioFramesReadyForDecode.IsEmpty() &&
+			LoadedVideoFrames.IsEmpty() && LoadedAudioFrames.IsEmpty())
+		{
+			CurrentTime = MovieDuration;
+			CurrentState = EMediaState::Paused;
+			OutEvents.Push(EMediaEvent::PlaybackEndReached);
+			OutEvents.Push(EMediaEvent::PlaybackSuspended);
+		}
+	}
 }
 
 /* IMediaTracks interface
@@ -298,7 +709,7 @@ void FWebMMediaPlayer::TickInput(FTimespan DeltaTime, FTimespan /*Timecode*/)
 
 bool FWebMMediaPlayer::GetAudioTrackFormat(int32 TrackIndex, int32 FormatIndex, FMediaAudioTrackFormat& OutFormat) const
 {
-	if (TrackIndex < 0 || FormatIndex <= 0)
+	if (FormatIndex != 0)
 	{
 		// We only support one format per track
 		return false;
@@ -313,7 +724,6 @@ bool FWebMMediaPlayer::GetAudioTrackFormat(int32 TrackIndex, int32 FormatIndex,
 	OutFormat.NumChannels = AudioTracks[TrackIndex]->GetChannels();
 	OutFormat.SampleRate = AudioTracks[TrackIndex]->GetSamplingRate();
 	OutFormat.TypeName = UTF8_TO_TCHAR(AudioTracks[TrackIndex]->GetCodecNameAsUTF8());
-
 	return true;
 }
 
@@ -407,57 +817,50 @@ bool FWebMMediaPlayer::GetVideoTrackFormat(int32 TrackIndex, int32 FormatIndex,
 	}
 
 	OutFormat.Dim = FIntPoint(VideoTracks[TrackIndex]->GetWidth(), VideoTracks[TrackIndex]->GetHeight());
-	OutFormat.FrameRate = VideoTracks[TrackIndex]->GetFrameRate();
-	OutFormat.TypeName = VideoTracks[TrackIndex]->GetCodecNameAsUTF8();
-	OutFormat.FrameRates = TRange<float>(OutFormat.FrameRate, OutFormat.FrameRate);
-
+	OutFormat.FrameRate = VideoFrameRate.IsValid() ? VideoFrameRate.AsDecimal() : 0.0f;
+	OutFormat.TypeName = UTF8_TO_TCHAR(VideoTracks[TrackIndex]->GetCodecNameAsUTF8());
+	OutFormat.FrameRates = TRange<float>{ OutFormat.FrameRate };
 	return true;
 }
 
-bool FWebMMediaPlayer::SelectTrack(EMediaTrackType TrackType, int32 TrackIndex)
+bool FWebMMediaPlayer::SelectTrack(EMediaTrackType InTrackType, int32 InTrackIndex)
 {
-	bool bHasChanged = false;
-
-	if (TrackType == EMediaTrackType::Audio)
-	{
-		bHasChanged = SelectedAudioTrack != TrackIndex;
-		SelectedAudioTrack = TrackIndex;
-	}
-	else if (TrackType == EMediaTrackType::Video)
+	if (InTrackType == EMediaTrackType::Video && InTrackIndex != SelectedVideoTrack && (InTrackIndex == INDEX_NONE || (InTrackIndex >= 0 && InTrackIndex < VideoTracks.Num())))
 	{
-		bHasChanged = SelectedVideoTrack != TrackIndex;
-		SelectedVideoTrack = TrackIndex;
-	}
-
-	if (bHasChanged)
-	{
-		if (CurrentState == EMediaState::Playing)
+		bPendingVideoTrackChange = true;
+		SelectedVideoTrack = InTrackIndex;
+		if (SelectedVideoTrack != INDEX_NONE)
 		{
-			Seek(CurrentTime);
-		}
+			const mkvparser::VideoTrack* VideoTrack = VideoTracks[SelectedVideoTrack];
 
-		if (CurrentState != EMediaState::Closed)
-		{
-			if (SelectedAudioTrack != INDEX_NONE)
+			auto DefDurNanos = VideoTrack->GetDefaultDuration();
+			if (DefDurNanos == 0)
 			{
-				const mkvparser::AudioTrack* AudioTrack = AudioTracks[SelectedAudioTrack];
-
-				size_t PrivateDataSize;
-				const uint8* PrivateData = AudioTrack->GetCodecPrivate(PrivateDataSize);
-
-				AudioDecoder->Initialize(AudioTrack->GetCodecId(), AudioTrack->GetSamplingRate(), AudioTrack->GetChannels(), PrivateData, PrivateDataSize);
+				UE_LOG(LogWebMMedia, Error, TEXT("Video track default duration not present in file"));
 			}
-
-			if (SelectedVideoTrack != INDEX_NONE)
+			VideoFrameRate.Numerator = ETimespan::TicksPerSecond;
+			VideoFrameRate.Denominator = DefDurNanos / 100;
+			// Is the value reasonable? (less than 60fps)
+			if (VideoFrameRate.IsValid() && VideoFrameRate.AsDecimal() > 60.5)
 			{
-				const mkvparser::VideoTrack* VideoTrack = VideoTracks[SelectedVideoTrack];
-
-				VideoDecoder->Initialize(VideoTrack->GetCodecId());
+				VideoFrameRate = FFrameRate(30,1);
+				bRecalculateFrameRate = true;
 			}
 		}
+		else
+		{
+			VideoFrameRate.Numerator = 0;
+			VideoFrameRate.Denominator = 1;
+		}
+		return true;
 	}
-
-	return true;
+	else if (InTrackType == EMediaTrackType::Audio && InTrackIndex != SelectedAudioTrack && (InTrackIndex == INDEX_NONE || (InTrackIndex >= 0 && InTrackIndex < AudioTracks.Num())))
+	{
+		bPendingAudioTrackChange = true;
+		SelectedAudioTrack = InTrackIndex;
+		return true;
+	}
+	return false;
 }
 
 bool FWebMMediaPlayer::SetTrackFormat(EMediaTrackType TrackType, int32 TrackIndex, int32 FormatIndex)
@@ -480,35 +883,27 @@ bool FWebMMediaPlayer::CanControl(EMediaControl Control) const
 {
 	if (Control == EMediaControl::Pause)
 	{
-		return (CurrentState == EMediaState::Playing);
+		return (CurrentState == EMediaState::Playing || CurrentState == EMediaState::Preparing);
 	}
-
-	if (Control == EMediaControl::Resume)
+	else if (Control == EMediaControl::Resume)
 	{
-		return ((CurrentState == EMediaState::Playing) || (CurrentState == EMediaState::Stopped));
+		return (CurrentState == EMediaState::Playing) || (CurrentState == EMediaState::Stopped);
 	}
-
-	if (Control == EMediaControl::Seek)
+	else if (Control == EMediaControl::Seek)
 	{
-		return ((CurrentState != EMediaState::Closed) && (CurrentState != EMediaState::Error));
+		return (CurrentState != EMediaState::Closed) && (CurrentState != EMediaState::Error);
 	}
-
 	return false;
 }
 
 FTimespan FWebMMediaPlayer::GetDuration() const
 {
-	if (CurrentState == EMediaState::Error || CurrentState == EMediaState::Closed)
-	{
-		return FTimespan::Zero();
-	}
-
-	return FTimespan::FromMicroseconds((double) MkvSegment->GetDuration() / 1000.0);
+	return MovieDuration;
 }
 
 float FWebMMediaPlayer::GetRate() const
 {
-	return (CurrentState == EMediaState::Playing) ? 1.0f : 0.0f;
+	return CurrentState == EMediaState::Playing ? 1.0f : 0.0f;
 }
 
 EMediaState FWebMMediaPlayer::GetState() const
@@ -518,26 +913,19 @@ EMediaState FWebMMediaPlayer::GetState() const
 
 EMediaStatus FWebMMediaPlayer::GetStatus() const
 {
-	return EMediaStatus::None;
+	return CurrentState == EMediaState::Preparing ? EMediaStatus::Buffering : EMediaStatus::None;
 }
 
 TRangeSet<float> FWebMMediaPlayer::GetSupportedRates(EMediaRateThinning /*Thinning*/) const
 {
 	TRangeSet<float> Result;
-
 	Result.Add(TRange<float>(0.0f));
 	Result.Add(TRange<float>(1.0f));
-
 	return Result;
 }
 
 FTimespan FWebMMediaPlayer::GetTime() const
 {
-	if ((CurrentState == EMediaState::Closed) || (CurrentState == EMediaState::Error))
-	{
-		return FTimespan::Zero();
-	}
-
 	return CurrentTime;
 }
 
@@ -546,52 +934,31 @@ bool FWebMMediaPlayer::IsLooping() const
 	return bLooping;
 }
 
-bool FWebMMediaPlayer::Seek(const FTimespan& Time)
+bool FWebMMediaPlayer::Seek(const FTimespan& InNewTime, const FMediaSeekParams& InAdditionalParams)
 {
-	if ((CurrentState == EMediaState::Closed) ||
-		(CurrentState == EMediaState::Error) ||
-		(CurrentState == EMediaState::Preparing))
+	if (CurrentState == EMediaState::Closed || CurrentState == EMediaState::Error)
 	{
 		UE_LOG(LogWebMMedia, Warning, TEXT("Cannot seek while closed or in error state"));
 		return false;
 	}
-
-	MkvSeekToTime(Time);
-
-	if (SelectedVideoTrack != INDEX_NONE)
-	{
-		const mkvparser::VideoTrack* VideoTrack = VideoTracks[SelectedVideoTrack];
-
-		VideoDecoder->Initialize(VideoTrack->GetCodecId());
-	}
-
-	if (SelectedAudioTrack != INDEX_NONE)
-	{
-		const mkvparser::AudioTrack* AudioTrack = AudioTracks[SelectedAudioTrack];
-
-		size_t PrivateDataSize;
-		const uint8* PrivateData = AudioTrack->GetCodecPrivate(PrivateDataSize);
-
-		AudioDecoder->Initialize(AudioTrack->GetCodecId(), AudioTrack->GetSamplingRate(), AudioTrack->GetChannels(), PrivateData, PrivateDataSize);
-	}
-
-	CurrentTime = Time;
-	bNoMoreToRead = false;
-
-	EventSink.ReceiveMediaEvent(EMediaEvent::SeekCompleted);
-
+	SeekToTime = InNewTime;
+	CurrentSeekIndex = InAdditionalParams.NewSequenceIndex.Get(0);
 	return true;
 }
 
 bool FWebMMediaPlayer::SetLooping(bool Looping)
 {
 	bLooping = Looping;
-
 	return true;
 }
 
 bool FWebMMediaPlayer::SetRate(float Rate)
 {
+	if (CurrentState == EMediaState::Closed || CurrentState == EMediaState::Error)
+	{
+		return false;
+	}
+
 	if (Rate == 0.0f)
 	{
 		Pause();
@@ -608,6 +975,36 @@ bool FWebMMediaPlayer::SetRate(float Rate)
 	}
 }
 
+bool FWebMMediaPlayer::GetPlayerFeatureFlag(EFeatureFlag InFlag) const
+{
+	switch(InFlag)
+	{
+		case EFeatureFlag::UsePlaybackTimingV2:
+		{
+			return true;
+		}
+		case EFeatureFlag::PlayerUsesInternalFlushOnSeek:
+		{
+			return true;
+		}
+		default:
+		{
+			return IMediaPlayer::GetPlayerFeatureFlag(InFlag);
+		}
+	}
+}
+
+bool FWebMMediaPlayer::FlushOnSeekStarted() const
+{
+	return false;
+}
+
+bool FWebMMediaPlayer::FlushOnSeekCompleted() const
+{
+	return false;
+}
+
+
 bool FWebMMediaPlayer::MkvRead()
 {
 	int64 FilePosition = 0;
@@ -624,30 +1021,30 @@ bool FWebMMediaPlayer::MkvRead()
 
 	MkvSegment.Reset(Segment);
 
-	if (MkvSegment->Load() < 0)
+	if (MkvSegment->Load() < 0 || !MkvSegment->GetInfo())
 	{
 		return false;
 	}
 
-	// Read all the tracks
+	MovieDuration = FTimespan::FromMicroseconds(MkvSegment->GetInfo()->GetDuration() / 1000);
 
+	// Read all the tracks
 	const mkvparser::Tracks* Tracks = MkvSegment->GetTracks();
 	int32 NumOfTracks = Tracks->GetTracksCount();
-
 	if (NumOfTracks == 0)
 	{
 		UE_LOG(LogWebMMedia, Warning, TEXT("File doesn't have any tracks"));
 		return false;
 	}
 
-	for (int32 i = 0; i < NumOfTracks; ++i)
+	for(int32 i=0; i<NumOfTracks; ++i)
 	{
 		const mkvparser::Track* Track = Tracks->GetTrackByIndex(i);
 		check(Track);
 
 		if (Track->GetType() == mkvparser::Track::kVideo)
 		{
-			if (FCStringAnsi::Strcmp(Track->GetCodecId(), "V_VP8") == 0 || FCStringAnsi::Strcmp(Track->GetCodecId(), "V_VP9") == 0)
+			if (GetCodecFromMKVTrack(Track->GetCodecId()) != ETrackCodec::Undefined)
 			{
 				VideoTracks.Add(static_cast<const mkvparser::VideoTrack*>(Track));
 			}
@@ -659,7 +1056,7 @@ bool FWebMMediaPlayer::MkvRead()
 		}
 		else if (Track->GetType() == mkvparser::Track::kAudio)
 		{
-			if (FCStringAnsi::Strcmp(Track->GetCodecId(), "A_OPUS") == 0 || FCStringAnsi::Strcmp(Track->GetCodecId(), "A_VORBIS") == 0)
+			if (GetCodecFromMKVTrack(Track->GetCodecId()) != ETrackCodec::Undefined)
 			{
 				AudioTracks.Add(static_cast<const mkvparser::AudioTrack*>(Track));
 			}
@@ -676,24 +1073,32 @@ bool FWebMMediaPlayer::MkvRead()
 		}
 	}
 
-	if (VideoTracks.Num() == 0 || AudioTracks.Num() == 0)
+	if (VideoTracks.IsEmpty() && AudioTracks.IsEmpty())
 	{
-		UE_LOG(LogWebMMedia, Warning, TEXT("File doesn't have video or audio. Right now only files with both are supported"));
+		UE_LOG(LogWebMMedia, Warning, TEXT("File doesn't have usable video or audio."));
 		return false;
 	}
-
 	return true;
 }
 
 void FWebMMediaPlayer::MkvSeekToNextValidBlock()
 {
-	for (;;)
+	while(1)
 	{
 		if (!MkvCurrentCluster)
 		{
+			if (!MkvSegment)
+			{
+				bFileError = true;
+				return;
+			}
 			MkvCurrentCluster = MkvSegment->GetFirst();
 			MkvCurrentBlockEntry = nullptr;
-			check(MkvCurrentCluster);
+			if (!MkvCurrentCluster)
+			{
+				bFileError = true;
+				return;
+			}
 		}
 
 		if (!MkvCurrentBlockEntry || MkvCurrentBlockEntry->EOS())
@@ -701,16 +1106,16 @@ void FWebMMediaPlayer::MkvSeekToNextValidBlock()
 			if (MkvCurrentCluster->GetFirst(MkvCurrentBlockEntry) != 0)
 			{
 				UE_LOG(LogWebMMedia, Warning, TEXT("Something went wrong while seeking"));
-				bNoMoreToRead = true;
+				bFileError = true;
 				return;
 			}
 		}
 		else
 		{
-			if(MkvCurrentCluster->GetNext(MkvCurrentBlockEntry, MkvCurrentBlockEntry) != 0)
+			if (MkvCurrentCluster->GetNext(MkvCurrentBlockEntry, MkvCurrentBlockEntry) != 0)
 			{
 				UE_LOG(LogWebMMedia, Warning, TEXT("Something went wrong while seeking"));
-				bNoMoreToRead = true;
+				bFileError = true;
 				return;
 			}
 		}
@@ -719,11 +1124,14 @@ void FWebMMediaPlayer::MkvSeekToNextValidBlock()
 		{
 			MkvCurrentBlockEntry = nullptr;
 			MkvCurrentCluster = MkvSegment->GetNext(MkvCurrentCluster);
-			check(MkvCurrentCluster);
-
+			if (!MkvCurrentCluster)
+			{
+				bFileError = true;
+				return;
+			}
 			if (MkvCurrentCluster->EOS())
 			{
-				bNoMoreToRead = true;
+				bFileReachedEOS = true;
 				return;
 			}
 			else
@@ -731,21 +1139,18 @@ void FWebMMediaPlayer::MkvSeekToNextValidBlock()
 				continue;
 			}
 		}
-
 		return;
 	}
 }
 
-void FWebMMediaPlayer::MkvSeekToTime(const FTimespan& Time)
+void FWebMMediaPlayer::MkvSeekToTime(const FTimespan& InTime)
 {
-	// TODO: for more precise seeking, we should use CUEs
-
-	uint64 TimeInNs = (uint64_t) (Time.GetTotalMicroseconds() * 1000.0);
-
-	MkvCurrentCluster = MkvSegment->FindCluster(TimeInNs);
-	MkvCurrentBlockEntry = nullptr;
-
-	Samples->FlushSamples();
+	if (MkvSegment)
+	{
+		uint64 TimeInNs = (uint64)InTime.GetTotalMicroseconds() * 1000U;
+		MkvCurrentCluster = MkvSegment->FindCluster(TimeInNs);
+		MkvCurrentBlockEntry = nullptr;
+	}
 }
 
 void FWebMMediaPlayer::Resume()
@@ -760,33 +1165,56 @@ void FWebMMediaPlayer::Pause()
 	OutEvents.Push(EMediaEvent::PlaybackSuspended);
 }
 
-void FWebMMediaPlayer::Stop()
+void FWebMMediaPlayer::AddVideoSampleFromDecodingThread(TSharedRef<FWebMMediaTextureSample, ESPMode::ThreadSafe> InSample)
 {
-	Pause();
+	if (InSample->GetDecoderIndex() == CurrentDecoderIndex)
+	{
+		InSample->SetShutdownPoolableDelegate(VideoFrameDoneDelegate);
+		FScopeLock lock(&DecodedSamplesLock);
+		DecodedVideoFrames.Emplace(InSample);
+	}
 }
 
-void FWebMMediaPlayer::AddVideoSampleFromDecodingThread(TSharedRef<FWebMMediaTextureSample, ESPMode::ThreadSafe> Sample)
+void FWebMMediaPlayer::AddAudioSampleFromDecodingThread(TSharedRef<FWebMMediaAudioSample, ESPMode::ThreadSafe> InSample)
 {
-	if (Sample->GetTime().Time < CurrentTime)
+	if (InSample->GetDecoderIndex() == CurrentDecoderIndex)
 	{
-		// We don't care about expired samples
-		return;
+		InSample->SetShutdownPoolableDelegate(AudioFrameDoneDelegate);
+		FScopeLock lock(&DecodedSamplesLock);
+		DecodedAudioFrames.Emplace(InSample);
 	}
+}
 
-	Samples->AddVideo(Sample);
+void FWebMMediaPlayer::ReportVideoDecodingError(FString InErrorMessage)
+{
+	bDecodeError = true;
 }
 
-void FWebMMediaPlayer::AddAudioSampleFromDecodingThread(TSharedRef<FWebMMediaAudioSample, ESPMode::ThreadSafe> Sample)
+void FWebMMediaPlayer::ReportAudioDecodingError(FString InErrorMessage)
 {
-	if (Sample->GetTime().Time < CurrentTime)
+	bDecodeError = true;
+}
+
+
+void FWebMMediaPlayer::ReturnVideoSample(const FWebMMediaTextureSample* InSample)
+{
+	if (InSample->GetDecoderIndex() == CurrentDecoderIndex)
 	{
-		// We don't care about expired samples
-		return;
+		FScopeLock lock(&TimeLock);
+		CurrentVideoTime = InSample->GetTime().Time;
 	}
+}
 
-	Samples->AddAudio(Sample);
+void FWebMMediaPlayer::ReturnAudioSample(const FWebMMediaAudioSample* InSample)
+{
+	if (InSample->GetDecoderIndex() == CurrentDecoderIndex)
+	{
+		FScopeLock lock(&TimeLock);
+		CurrentAudioTime = InSample->GetTime().Time;
+	}
 }
 
+
 #undef LOCTEXT_NAMESPACE
 
 #endif // WITH_WEBM_LIBS
diff --git a/Engine/Plugins/Media/WebMMedia/Source/WebMMedia/Private/Player/WebMMediaPlayer.h b/Engine/Plugins/Media/WebMMedia/Source/WebMMedia/Private/Player/WebMMediaPlayer.h
--- a/Engine/Plugins/Media/WebMMedia/Source/WebMMedia/Private/Player/WebMMediaPlayer.h
+++ b/Engine/Plugins/Media/WebMMedia/Source/WebMMedia/Private/Player/WebMMediaPlayer.h
@@ -7,6 +7,8 @@
 #include "CoreMinimal.h"
 #include "Containers/Array.h"
 #include "Containers/UnrealString.h"
+#include "Templates/SharedPointer.h"
+#include "Misc/Optional.h"
 #include "Misc/Guid.h"
 #include "IMediaCache.h"
 #include "IMediaControls.h"
@@ -15,12 +17,18 @@
 #include "IMediaView.h"
 #include "IMediaEventSink.h"
 #include "Misc/Timespan.h"
+#include "Misc/FrameRate.h"
 #include "MkvFileReader.h"
 #include "WebMSamplesSink.h"
+#include "WebMMediaTextureSample.h"
+#include "WebMMediaAudioSample.h"
+#include <atomic>
+
 
 class FMediaSamples;
 class FWebMVideoDecoder;
 class FWebMAudioDecoder;
+struct FWebMFrame;
 
 class FWebMMediaPlayer
 	: public IMediaPlayer
@@ -38,6 +46,8 @@ public:
 	//~ IWebMSamplesSink interface
 	virtual void AddVideoSampleFromDecodingThread(TSharedRef<FWebMMediaTextureSample, ESPMode::ThreadSafe> Sample) override;
 	virtual void AddAudioSampleFromDecodingThread(TSharedRef<FWebMMediaAudioSample, ESPMode::ThreadSafe> Sample) override;
+	virtual void ReportVideoDecodingError(FString InErrorMessage) override;
+	virtual void ReportAudioDecodingError(FString InErrorMessage) override;
 
 public:
 	//~ IMediaPlayer interface
@@ -53,7 +63,6 @@ public:
 	virtual IMediaView& GetView() override;
 	virtual bool Open(const FString& Url, const IMediaOptions* Options) override;
 	virtual bool Open(const TSharedRef<FArchive, ESPMode::ThreadSafe>& Archive, const FString& OriginalUrl, const IMediaOptions* Options) override;
-	virtual void TickFetch(FTimespan DeltaTime, FTimespan Timecode) override;
 	virtual void TickInput(FTimespan DeltaTime, FTimespan Timecode) override;
 
 protected:
@@ -66,9 +75,14 @@ protected:
 	virtual TRangeSet<float> GetSupportedRates(EMediaRateThinning Thinning) const override;
 	virtual FTimespan GetTime() const override;
 	virtual bool IsLooping() const override;
-	virtual bool Seek(const FTimespan& Time) override;
+	virtual bool Seek(const FTimespan& Time) override
+	{ check(!"You have to call the override with additional options!"); return false; }
+	virtual bool Seek(const FTimespan& InNewTime, const FMediaSeekParams& InAdditionalParams) override;
 	virtual bool SetLooping(bool Looping) override;
 	virtual bool SetRate(float Rate) override;
+	virtual bool GetPlayerFeatureFlag(EFeatureFlag InFlag) const override;
+	virtual bool FlushOnSeekStarted() const override;
+	virtual bool FlushOnSeekCompleted() const override;
 
 protected:
 	//~ IMediaTracks interface
@@ -85,32 +99,90 @@ protected:
 	virtual bool SetTrackFormat(EMediaTrackType TrackType, int32 TrackIndex, int32 FormatIndex) override;
 
 private:
-	TArray<EMediaEvent> OutEvents;
+	enum class ETrackCodec
+	{
+		Undefined,
+		VP8,
+		VP9,
+		Vorbis,
+		Opus
+	};
+
+	struct FSeekTo
+	{
+		FTimespan TargetTime;
+		bool bAwaitTime = false;
+		bool bGotVideoKeyframe = false;
+		void Reset()
+		{
+			TargetTime = FTimespan::Zero();
+			bAwaitTime = false;
+			bGotVideoKeyframe = false;
+		}
+	};
+
+	static ETrackCodec GetCodecFromMKVTrack(const char* InMKVCodecName);
+	static FString GetMKVCodec(ETrackCodec InCodec);
+	void TickInternal();
+	void ReturnVideoSample(const FWebMMediaTextureSample* InSample);
+	void ReturnAudioSample(const FWebMMediaAudioSample* InSample);
+	void Resume();
+	void Pause();
+	bool MkvRead();
+	void MkvSeekToNextValidBlock();
+	void MkvSeekToTime(const FTimespan& Time);
+
+
 	IMediaEventSink& EventSink;
+	FCriticalSection MediaSamplesLock;
 	TSharedPtr<FMediaSamples, ESPMode::ThreadSafe> Samples;
-	TUniquePtr<FWebMVideoDecoder> VideoDecoder;
-	TUniquePtr<FWebMAudioDecoder> AudioDecoder;
+	TArray<EMediaEvent> OutEvents;
+	uint32 CurrentDecoderIndex { 0 };
+	FString MediaUrl;
 	TArray<const mkvparser::VideoTrack*> VideoTracks;
 	TArray<const mkvparser::AudioTrack*> AudioTracks;
+	TUniquePtr<FWebMVideoDecoder> VideoDecoder;
+	TUniquePtr<FWebMAudioDecoder> AudioDecoder;
 	TUniquePtr<FMkvFileReader> MkvReader;
 	TUniquePtr<mkvparser::Segment> MkvSegment;
-	const mkvparser::Cluster* MkvCurrentCluster;
-	const mkvparser::BlockEntry* MkvCurrentBlockEntry;
-	FString MediaUrl;
-	EMediaState CurrentState;
-	int32 SelectedAudioTrack;
-	int32 SelectedVideoTrack;
+	const mkvparser::Cluster* MkvCurrentCluster { };
+	const mkvparser::BlockEntry* MkvCurrentBlockEntry { };
+	EMediaState CurrentState { EMediaState::Closed };
+	TOptional<ETrackCodec> NewPendingVideoCodec;
+	TOptional<ETrackCodec> NewPendingAudioCodec;
+	int32 SelectedVideoTrack { INDEX_NONE };
+	int32 SelectedAudioTrack { INDEX_NONE };
+	bool bPendingVideoTrackChange { false };
+	bool bPendingAudioTrackChange { false };
 	FTimespan CurrentTime;
-	bool bNoMoreToRead;
-	bool bLooping;
+	TOptional<FTimespan> SeekToTime;
+	FSeekTo SeekTo;
 
-private:
-	void Resume();
-	void Pause();
-	void Stop();
-	bool MkvRead();
-	void MkvSeekToNextValidBlock();
-	void MkvSeekToTime(const FTimespan& Time);
+	int32 CurrentSeekIndex { 0 };
+	int32 CurrentLoopIndex { 0 };
+	bool bFileReachedEOS { false };
+	bool bFileError { false };
+	bool bDecodeError { false };
+	bool bLooping { false} ;
+
+	TArray<TSharedPtr<FWebMFrame, ESPMode::ThreadSafe>> VideoFramesReadyForDecode;
+	TArray<TSharedPtr<FWebMFrame, ESPMode::ThreadSafe>> AudioFramesReadyForDecode;
+	TArray<TSharedPtr<FWebMFrame, ESPMode::ThreadSafe>> LoadedVideoFrames;
+	TArray<TSharedPtr<FWebMFrame, ESPMode::ThreadSafe>> LoadedAudioFrames;
+	FTimespan MovieDuration;
+	FFrameRate VideoFrameRate {0,0};
+	bool bRecalculateFrameRate = false;
+	FCriticalSection TimeLock;
+	TOptional<FTimespan> CurrentVideoTime;
+	TOptional<FTimespan> CurrentAudioTime;
+
+	FCriticalSection DecodedSamplesLock;
+	std::atomic<int32> NumDecodingEnqueuedVideoFrames;
+	std::atomic<int32> NumDecodingEnqueuedAudioFrames;
+	TArray<TSharedPtr<FWebMMediaTextureSample, ESPMode::ThreadSafe>> DecodedVideoFrames;
+	TArray<TSharedPtr<FWebMMediaAudioSample, ESPMode::ThreadSafe>> DecodedAudioFrames;
+	FWebMMediaTextureSample::FShutdownPoolableDlg VideoFrameDoneDelegate;
+	FWebMMediaAudioSample::FShutdownPoolableDlg AudioFrameDoneDelegate;
 };
 
 #endif // WITH_WEBM_LIBS
diff --git a/Engine/Plugins/Media/WebMMedia/Source/WebMMedia/Private/Player/WebMMediaTextureSample.cpp b/Engine/Plugins/Media/WebMMedia/Source/WebMMedia/Private/Player/WebMMediaTextureSample.cpp
--- a/Engine/Plugins/Media/WebMMedia/Source/WebMMedia/Private/Player/WebMMediaTextureSample.cpp
+++ b/Engine/Plugins/Media/WebMMedia/Source/WebMMedia/Private/Player/WebMMediaTextureSample.cpp
@@ -8,12 +8,13 @@ FWebMMediaTextureSample::FWebMMediaTextureSample()
 {
 }
 
-void FWebMMediaTextureSample::Initialize(FIntPoint InDisplaySize, FIntPoint InTotalSize, FTimespan InTime, FTimespan InDuration)
+void FWebMMediaTextureSample::Initialize(FIntPoint InDisplaySize, FIntPoint InTotalSize, FMediaTimeStamp InTime, FTimespan InDuration, uint32 InDecoderIndex)
 {
 	Time = InTime;
 	DisplaySize = InDisplaySize;
 	TotalSize = InTotalSize;
 	Duration = InDuration;
+	DecoderIndex = InDecoderIndex;
 }
 
 void FWebMMediaTextureSample::CreateTexture(FRHICommandListBase& RHICmdList)
@@ -87,10 +88,12 @@ bool FWebMMediaTextureSample::IsOutputSrgb() const
 
 void FWebMMediaTextureSample::ShutdownPoolable()
 {
+	ShutdownDelegate.ExecuteIfBound(this);
+
 	// Drop reference to the texture. It should be released by the outside system.
 	Texture = nullptr;
 	
-	Time = FTimespan::Zero();
+	Time = FMediaTimeStamp(FTimespan::Zero());
 }
 
 TRefCountPtr<FRHITexture> FWebMMediaTextureSample::GetTextureRef() const
diff --git a/Engine/Plugins/Media/WebMMedia/Source/WebMMedia/Private/Player/WebMVideoDecoder.cpp b/Engine/Plugins/Media/WebMMedia/Source/WebMMedia/Private/Player/WebMVideoDecoder.cpp
--- a/Engine/Plugins/Media/WebMMedia/Source/WebMMedia/Private/Player/WebMVideoDecoder.cpp
+++ b/Engine/Plugins/Media/WebMMedia/Source/WebMMedia/Private/Player/WebMVideoDecoder.cpp
@@ -62,6 +62,8 @@ namespace
 FWebMVideoDecoder::FWebMVideoDecoder(IWebMSamplesSink& InSamples)
 	: VideoSamplePool(new FWebMMediaTextureSamplePool)
 	, Samples(InSamples)
+	, Codec(ECodec::Unknown)
+	, CurrentSequenceIndex(0)
 	, bTexturesCreated(false)
 	, bIsInitialized(false)
 {
@@ -75,29 +77,36 @@ FWebMVideoDecoder::~FWebMVideoDecoder()
 bool FWebMVideoDecoder::Initialize(const char* CodecName)
 {
 	Close();
+	return InitializeInternal(FString(CodecName));
+}
 
+bool FWebMVideoDecoder::InitializeInternal(FString InCodecName)
+{
 	const int32 NumOfThreads = 1;
-	const vpx_codec_dec_cfg_t CodecConfig = { NumOfThreads, 0, 0 };
-	if (FCStringAnsi::Strcmp(CodecName, "V_VP8") == 0)
+	CodecConfig.threads = NumOfThreads;
+	CodecConfig.w = 0;
+	CodecConfig.h = 0;
+	if (InCodecName.Equals(TEXT("V_VP8")))
 	{
 		verify(vpx_codec_dec_init(&Context, vpx_codec_vp8_dx(), &CodecConfig, /*VPX_CODEC_USE_FRAME_THREADING*/ 0) == 0);
+		Codec = ECodec::VP8;
 	}
-	else if (FCStringAnsi::Strcmp(CodecName, "V_VP9") == 0)
+	else if (InCodecName.Equals(TEXT("V_VP9")))
 	{
 		verify(vpx_codec_dec_init(&Context, vpx_codec_vp9_dx(), &CodecConfig, /*VPX_CODEC_USE_FRAME_THREADING*/ 0) == 0);
+		Codec = ECodec::VP9;
 	}
 	else
 	{
-		UE_LOG(LogWebMMedia, Display, TEXT("Unsupported video codec: %hs"), CodecName);
+		UE_LOG(LogWebMMedia, Display, TEXT("Unsupported video codec: %s"), *InCodecName);
+		Codec = ECodec::Unknown;
 		return false;
 	}
-
 	bIsInitialized = true;
-
 	return true;
 }
 
-void FWebMVideoDecoder::DecodeVideoFramesAsync(const TArray<TSharedPtr<FWebMFrame>>& VideoFrames)
+void FWebMVideoDecoder::DecodeVideoFramesAsync(const TArray<TSharedPtr<FWebMFrame, ESPMode::ThreadSafe>>& VideoFrames)
 {
 	static bool bUseRenderThread = FPlatformMisc::UseRenderThread();
 	if (bUseRenderThread)
@@ -123,24 +132,79 @@ bool FWebMVideoDecoder::IsBusy() const
 	return VideoDecodingTask && !VideoDecodingTask->IsComplete();
 }
 
-void FWebMVideoDecoder::DoDecodeVideoFrames(const TArray<TSharedPtr<FWebMFrame>>& VideoFrames)
+void FWebMVideoDecoder::DoDecodeVideoFrames(const TArray<TSharedPtr<FWebMFrame, ESPMode::ThreadSafe>>& VideoFrames)
 {
-	for (const TSharedPtr<FWebMFrame>& VideoFrame : VideoFrames)
+	auto CreateDecoder = [&](const FString& InCodec) -> bool
+	{
+		if (bIsInitialized)
+		{
+			// If we have a decoder then we need to save it since its video frames
+			// may still be referenced by the image conversion RHI task.
+			FScopeLock lock(&PreviousContextsLock);
+			PreviousContexts.Emplace(Context);
+			FMemory::Memzero(Context);
+			bIsInitialized = false;
+		}
+		return InitializeInternal(InCodec);
+	};
+	for (const TSharedPtr<FWebMFrame, ESPMode::ThreadSafe>& VideoFrame : VideoFrames)
 	{
+		if (!bIsInitialized)
+		{
+			if (VideoFrame->CodecName.IsEmpty())
+			{
+				UE_LOG(LogWebMMedia, Display, TEXT("No codec specified to initialize the codec lazily"));
+				Samples.ReportVideoDecodingError(TEXT("No codec specified to initialize the codec lazily"));
+				return;
+			}
+		}
+		if (VideoFrame->CodecName.Len())
+		{
+			if (!CreateDecoder(VideoFrame->CodecName))
+			{
+				UE_LOG(LogWebMMedia, Display, TEXT("Failed to create codec lazily"));
+				Samples.ReportVideoDecodingError(TEXT("Failed to create codec lazily"));
+				return;
+			}
+		}
+
+		// Change in sequence?
+		if (VideoFrame->SequenceIndex != CurrentSequenceIndex)
+		{
+			if (!CreateDecoder(Codec == ECodec::VP8 ? TEXT("V_VP8"): TEXT("V_VP9")))
+			{
+				Samples.ReportVideoDecodingError(TEXT("Failed to reset codec"));
+				return;
+			}
+			CurrentSequenceIndex = VideoFrame->SequenceIndex;
+		}
+
 		if (vpx_codec_decode(&Context, VideoFrame->Data.GetData(), VideoFrame->Data.Num(), nullptr, 0) != 0)
 		{
 			UE_LOG(LogWebMMedia, Display, TEXT("Error decoding video frame"));
+			Samples.ReportVideoDecodingError(TEXT("Error decoding video frame"));
 			return;
 		}
 
 		const void* ImageIter = nullptr;
 		while (const vpx_image_t* Image = vpx_codec_get_frame(&Context, &ImageIter))
 		{
+			// We only handle 8 bit bitdepth for now.
+			if (Image->bit_depth > 8)
+			{
+				UE_LOG(LogWebMMedia, Display, TEXT("%d bit video frames are not supported"), Image->bit_depth);
+				Samples.ReportVideoDecodingError(FString::Printf(TEXT("%d bit video frames are not supported"), Image->bit_depth));
+				return;
+			}
 			FWebMVideoDecoder* Self = this;
 			if (!bTexturesCreated)
 			{
 				// First creation of conversion textures
-
+				/*
+					NOTE: If there is playback of a movie with several video tracks and switching between them
+					      on the fly is intended, then the conversion texture may need to be resized if the new
+					      track video is larger than what we had before.
+				*/
 				bTexturesCreated = true;
 				ENQUEUE_RENDER_COMMAND(WebMMediaPlayerCreateTextures)(
 					[Self, Image](FRHICommandListImmediate& RHICmdList)
@@ -151,7 +215,7 @@ void FWebMVideoDecoder::DoDecodeVideoFrames(const TArray<TSharedPtr<FWebMFrame>>
 
 			TSharedRef<FWebMMediaTextureSample, ESPMode::ThreadSafe> VideoSample = VideoSamplePool->AcquireShared();
 
-			VideoSample->Initialize(FIntPoint(Image->d_w, Image->d_h), FIntPoint(Image->d_w, Image->d_h), VideoFrame->Time, VideoFrame->Duration);
+			VideoSample->Initialize(FIntPoint(Image->d_w, Image->d_h), FIntPoint(Image->d_w, Image->d_h), FMediaTimeStamp(VideoFrame->Time, VideoFrame->SequenceIndex), VideoFrame->Duration, VideoFrame->DecoderIndex);
 
 			FConvertParams Params;
 			Params.VideoSample = VideoSample;
@@ -202,10 +266,24 @@ void FWebMVideoDecoder::Close()
 		vpx_codec_destroy(&Context);
 		bIsInitialized = false;
 	}
+	ResetPendingDecoders();
 
 	bTexturesCreated = false;
 }
 
+void FWebMVideoDecoder::ResetPendingDecoders()
+{
+	PreviousContextsLock.Lock();
+	TArray<vpx_codec_ctx_t> prev;
+	Swap(prev, PreviousContexts);
+	PreviousContextsLock.Unlock();
+	for(int32 i=0; i<prev.Num(); ++i)
+	{
+		vpx_codec_destroy(&prev[i]);
+	}
+}
+
+
 void FWebMVideoDecoder::ConvertYUVToRGBAndSubmit(FRHICommandListImmediate& CommandList, const FConvertParams& Params)
 {
 	TSharedPtr<FWebMMediaTextureSample, ESPMode::ThreadSafe> VideoSample = Params.VideoSample;
@@ -289,6 +367,7 @@ void FWebMVideoDecoder::ConvertYUVToRGBAndSubmit(FRHICommandListImmediate& Comma
 		CommandList.Transition(FRHITransitionInfo(RenderTarget, ERHIAccess::RTV, ERHIAccess::SRVGraphics));
 
 		Samples.AddVideoSampleFromDecodingThread(VideoSample.ToSharedRef());
+		ResetPendingDecoders();
 	}
 }
 
diff --git a/Engine/Plugins/Media/WebMMedia/Source/WebMMedia/Private/WebMContainer.cpp b/Engine/Plugins/Media/WebMMedia/Source/WebMMedia/Private/WebMContainer.cpp
--- a/Engine/Plugins/Media/WebMMedia/Source/WebMMedia/Private/WebMContainer.cpp
+++ b/Engine/Plugins/Media/WebMMedia/Source/WebMMedia/Private/WebMContainer.cpp
@@ -89,7 +89,7 @@ bool FWebMContainer::Open(const FString& FilePath)
 				{
 					SelectedVideoTrack = MkvFile->VideoTracks.Num() - 1;
 				}
-				
+
 			}
 			else
 			{
@@ -133,7 +133,7 @@ bool FWebMContainer::Open(const FString& FilePath)
 	return true;
 }
 
-void FWebMContainer::ReadFrames(FTimespan ReadBufferLength, TArray<TSharedPtr<FWebMFrame>>& AudioFrames, TArray<TSharedPtr<FWebMFrame>>& VideoFrames)
+void FWebMContainer::ReadFrames(FTimespan ReadBufferLength, TArray<TSharedPtr<FWebMFrame, ESPMode::ThreadSafe>>& AudioFrames, TArray<TSharedPtr<FWebMFrame, ESPMode::ThreadSafe>>& VideoFrames)
 {
 	if (bNoMoreToRead)
 	{
@@ -157,7 +157,7 @@ void FWebMContainer::ReadFrames(FTimespan ReadBufferLength, TArray<TSharedPtr<FW
 		{
 			const mkvparser::Block::Frame& MkvFrame = MkvFile->CurrentBlockEntry->GetBlock()->GetFrame(i);
 
-			TSharedPtr<FWebMFrame> Frame = MakeShared<FWebMFrame>();
+			TSharedPtr<FWebMFrame, ESPMode::ThreadSafe> Frame = MakeShared<FWebMFrame>();
 			Frame->Time = CurrentTime;
 			Frame->Duration = FMkvFileReader::GetVideoFrameDuration(*MkvFile->VideoTracks[SelectedVideoTrack]);
 			Frame->Data.SetNumUninitialized(MkvFrame.len);
diff --git a/Engine/Plugins/Media/WebMMedia/Source/WebMMedia/Public/WebMAudioDecoder.h b/Engine/Plugins/Media/WebMMedia/Source/WebMMedia/Public/WebMAudioDecoder.h
--- a/Engine/Plugins/Media/WebMMedia/Source/WebMMedia/Public/WebMAudioDecoder.h
+++ b/Engine/Plugins/Media/WebMMedia/Source/WebMMedia/Public/WebMAudioDecoder.h
@@ -22,33 +22,37 @@ public:
 
 public:
 	UE_API bool Initialize(const char* InCodec, int32 InSampleRate, int32 InChannels, const uint8* CodecPrivateData, size_t CodecPrivateDataSize);
-	UE_API void DecodeAudioFramesAsync(const TArray<TSharedPtr<FWebMFrame>>& AudioFrames);
+	UE_API void DecodeAudioFramesAsync(const TArray<TSharedPtr<FWebMFrame, ESPMode::ThreadSafe>>& AudioFrames);
 	UE_API bool IsBusy() const;
 
 private:
-	enum ESupportedCodecs
+	enum ECodec
 	{
+		Undefined,
 		Opus,
 		Vorbis,
 	};
 
 	struct FVorbisDecoder;
+	TArray<uint8> CodecSpecificData;
 	TUniquePtr<FWebMMediaAudioSamplePool> AudioSamplePool;
 	TUniquePtr<FVorbisDecoder> VorbisDecoder;
 	FGraphEventRef AudioDecodingTask;
 	TArray<uint8> DecodeBuffer;
 	IWebMSamplesSink& Samples;
-	OpusDecoder* OpusDecoder;
-	ESupportedCodecs Codec;
-	int32 FrameSize;
-	int32 SampleRate;
-	int32 Channels;
-
+	OpusDecoder* OpusDecoder { };
+	ECodec Codec { ECodec::Undefined };
+	int64 CurrentSequenceIndex { };
+	int32 FrameSize { };
+	int32 SampleRate { };
+	int32 Channels { };
+
+	bool InitializeInternal(FString InCodec, int32 InSampleRate, int32 InChannels, TConstArrayView<uint8> InCodecPrivate);
 	void InitializeOpus();
 	bool InitializeVorbis(const uint8* CodecPrivateData, size_t CodecPrivateDataSize);
-	void DoDecodeAudioFrames(const TArray<TSharedPtr<FWebMFrame>>& AudioFrames);
-	int32 DecodeOpus(const TSharedPtr<FWebMFrame>& AudioFrame);
-	int32 DecodeVorbis(const TSharedPtr<FWebMFrame>& AudioFrame);
+	void DoDecodeAudioFrames(const TArray<TSharedPtr<FWebMFrame, ESPMode::ThreadSafe>>& AudioFrames);
+	int32 DecodeOpus(const TSharedPtr<FWebMFrame, ESPMode::ThreadSafe>& AudioFrame);
+	int32 DecodeVorbis(const TSharedPtr<FWebMFrame, ESPMode::ThreadSafe>& AudioFrame);
 };
 
 #undef UE_API
diff --git a/Engine/Plugins/Media/WebMMedia/Source/WebMMedia/Public/WebMContainer.h b/Engine/Plugins/Media/WebMMedia/Source/WebMMedia/Public/WebMContainer.h
--- a/Engine/Plugins/Media/WebMMedia/Source/WebMMedia/Public/WebMContainer.h
+++ b/Engine/Plugins/Media/WebMMedia/Source/WebMMedia/Public/WebMContainer.h
@@ -35,7 +35,7 @@ public:
 	UE_API virtual ~FWebMContainer();
 
 	UE_API bool Open(const FString& File);
-	UE_API void ReadFrames(FTimespan AmountOfTimeToRead, TArray<TSharedPtr<FWebMFrame>>& AudioFrames, TArray<TSharedPtr<FWebMFrame>>& VideoFrames);
+	UE_API void ReadFrames(FTimespan AmountOfTimeToRead, TArray<TSharedPtr<FWebMFrame, ESPMode::ThreadSafe>>& AudioFrames, TArray<TSharedPtr<FWebMFrame, ESPMode::ThreadSafe>>& VideoFrames);
 	UE_API FWebMAudioTrackInfo GetCurrentAudioTrackInfo() const;
 	UE_API FWebMVideoTrackInfo GetCurrentVideoTrackInfo() const;
 
diff --git a/Engine/Plugins/Media/WebMMedia/Source/WebMMedia/Public/WebMMediaAudioSample.h b/Engine/Plugins/Media/WebMMedia/Source/WebMMedia/Public/WebMMediaAudioSample.h
--- a/Engine/Plugins/Media/WebMMedia/Source/WebMMedia/Public/WebMMediaAudioSample.h
+++ b/Engine/Plugins/Media/WebMMedia/Source/WebMMedia/Public/WebMMediaAudioSample.h
@@ -26,6 +26,7 @@ public:
 		, Duration(FTimespan::Zero())
 		, SampleRate(0)
 		, Time(FTimespan::Zero())
+		, DecoderIndex(0)
 	{ }
 
 	/** Virtual destructor. */
@@ -38,16 +39,20 @@ public:
 	 *
 	 * @param InBuffer The sample's data buffer.
 	 * @param InSize The size of the sample buffer (in bytes).
+	 * @param InChannels The number of channels
+	 * @param InSampleRate Sampling rate
 	 * @param InTime The sample time (relative to presentation clock).
 	 * @param InDuration The duration for which the sample is valid.
+	 * @param InDecoderIndex Opaque value identifying the producing decoder
 	 */
 	void Initialize(
 		const uint8* InBuffer,
 		uint32 InSize,
 		uint32 InChannels,
 		uint32 InSampleRate,
-		FTimespan InTime,
-		FTimespan InDuration)
+		FMediaTimeStamp InTime,
+		FTimespan InDuration,
+		uint32 InDecoderIndex)
 	{
 		check(InBuffer && InSize > 0);
 
@@ -58,8 +63,13 @@ public:
 		Duration = InDuration;
 		SampleRate = InSampleRate;
 		Time = InTime;
+		DecoderIndex = InDecoderIndex;
 	}
 
+	DECLARE_DELEGATE_OneParam(FShutdownPoolableDlg, const FWebMMediaAudioSample*);
+	void SetShutdownPoolableDelegate(FShutdownPoolableDlg InDelegate)
+	{ ShutdownDelegate = MoveTemp(InDelegate); }
+
 public:
 
 	//~ IMediaAudioSample interface
@@ -96,7 +106,7 @@ public:
 
 	virtual FMediaTimeStamp GetTime() const override
 	{
-		return FMediaTimeStamp(Time);
+		return Time;
 	}
 
 	const TArray<uint8>& GetDataBuffer() const
@@ -104,7 +114,14 @@ public:
 		return Buffer;
 	}
 
+	//~ IMediaPoolable interface
+	virtual void ShutdownPoolable() override
+	{ ShutdownDelegate.ExecuteIfBound(this); }
+
+	uint32 GetDecoderIndex() const
+	{ return DecoderIndex; }
 private:
+	FShutdownPoolableDlg ShutdownDelegate;
 
 	/** The sample's data buffer. */
 	TArray<uint8> Buffer;
@@ -119,7 +136,10 @@ private:
 	uint32 SampleRate;
 
 	/** Presentation time for which the sample was generated. */
-	FTimespan Time;
+	FMediaTimeStamp Time;
+
+	/** An identifier of the decoder that produced this sample */
+	uint32 DecoderIndex;
 };
 
 
diff --git a/Engine/Plugins/Media/WebMMedia/Source/WebMMedia/Public/WebMMediaFrame.h b/Engine/Plugins/Media/WebMMedia/Source/WebMMedia/Public/WebMMediaFrame.h
--- a/Engine/Plugins/Media/WebMMedia/Source/WebMMedia/Public/WebMMediaFrame.h
+++ b/Engine/Plugins/Media/WebMMedia/Source/WebMMedia/Public/WebMMediaFrame.h
@@ -9,4 +9,12 @@ struct FWebMFrame
 	TArray<uint8> Data;
 	FTimespan Time;
 	FTimespan Duration;
+	int64 SequenceIndex = 0;
+	uint32 DecoderIndex = 0;
+	bool bIsKeyframe = false;
+	// If codec name is set then a new decoder of this format is needed.
+	FString CodecName;
+	TArray<uint8> CodecSpecificData;
+	int32 NumChannels = 0;
+	int32 SamplingRate = 0;
 };
diff --git a/Engine/Plugins/Media/WebMMedia/Source/WebMMedia/Public/WebMMediaTextureSample.h b/Engine/Plugins/Media/WebMMedia/Source/WebMMedia/Public/WebMMediaTextureSample.h
--- a/Engine/Plugins/Media/WebMMedia/Source/WebMMedia/Public/WebMMediaTextureSample.h
+++ b/Engine/Plugins/Media/WebMMedia/Source/WebMMedia/Public/WebMMediaTextureSample.h
@@ -20,7 +20,7 @@ public:
 	UE_API FWebMMediaTextureSample();
 
 public:
-	UE_API void Initialize(FIntPoint InDisplaySize, FIntPoint InTotalSize, FTimespan InTime, FTimespan InDuration);
+	UE_API void Initialize(FIntPoint InDisplaySize, FIntPoint InTotalSize, FMediaTimeStamp InTime, FTimespan InDuration, uint32 InDecoderIndex);
 	UE_API void CreateTexture(FRHICommandListBase& RHICmdList);
 
 public:
@@ -36,6 +36,12 @@ public:
 	UE_API virtual bool IsCacheable() const override;
 	UE_API virtual bool IsOutputSrgb() const override;
 
+	DECLARE_DELEGATE_OneParam(FShutdownPoolableDlg, const FWebMMediaTextureSample*);
+	void SetShutdownPoolableDelegate(FShutdownPoolableDlg InDelegate)
+	{ ShutdownDelegate = MoveTemp(InDelegate); }
+	uint32 GetDecoderIndex() const
+	{ return DecoderIndex; }
+
 public:
 	//~ IMediaPoolable interface
 	UE_API virtual void ShutdownPoolable() override;
@@ -44,11 +50,13 @@ public:
 	UE_API TRefCountPtr<FRHITexture> GetTextureRef() const;
 
 private:
+	FShutdownPoolableDlg ShutdownDelegate;
 	TRefCountPtr<FRHITexture> Texture;
-	FTimespan Time;
+	FMediaTimeStamp Time;
 	FTimespan Duration;
 	FIntPoint TotalSize;
 	FIntPoint DisplaySize;
+	uint32 DecoderIndex = 0;
 };
 
 class FWebMMediaTextureSamplePool : public TMediaObjectPool<FWebMMediaTextureSample> { };
diff --git a/Engine/Plugins/Media/WebMMedia/Source/WebMMedia/Public/WebMSamplesSink.h b/Engine/Plugins/Media/WebMMedia/Source/WebMMedia/Public/WebMSamplesSink.h
--- a/Engine/Plugins/Media/WebMMedia/Source/WebMMedia/Public/WebMSamplesSink.h
+++ b/Engine/Plugins/Media/WebMMedia/Source/WebMMedia/Public/WebMSamplesSink.h
@@ -13,6 +13,8 @@ class IWebMSamplesSink
 public:
 	virtual void AddVideoSampleFromDecodingThread(TSharedRef<FWebMMediaTextureSample, ESPMode::ThreadSafe> Sample) = 0;
 	virtual void AddAudioSampleFromDecodingThread(TSharedRef<FWebMMediaAudioSample, ESPMode::ThreadSafe> Sample) = 0;
+	virtual void ReportVideoDecodingError(FString InErrorMessage) = 0;
+	virtual void ReportAudioDecodingError(FString InErrorMessage) = 0;
 
 	virtual ~IWebMSamplesSink() {}
 };
diff --git a/Engine/Plugins/Media/WebMMedia/Source/WebMMedia/Public/WebMVideoDecoder.h b/Engine/Plugins/Media/WebMMedia/Source/WebMMedia/Public/WebMVideoDecoder.h
--- a/Engine/Plugins/Media/WebMMedia/Source/WebMMedia/Public/WebMVideoDecoder.h
+++ b/Engine/Plugins/Media/WebMMedia/Source/WebMMedia/Public/WebMVideoDecoder.h
@@ -25,7 +25,7 @@ public:
 
 public:
 	WEBMMEDIA_API bool Initialize(const char* CodecName);
-	WEBMMEDIA_API void DecodeVideoFramesAsync(const TArray<TSharedPtr<FWebMFrame>>& VideoFrames);
+	WEBMMEDIA_API void DecodeVideoFramesAsync(const TArray<TSharedPtr<FWebMFrame, ESPMode::ThreadSafe>>& VideoFrames);
 	WEBMMEDIA_API bool IsBusy() const;
 
 private:
@@ -35,20 +35,34 @@ private:
 		const vpx_image_t* Image;
 	};
 
+	enum class ECodec
+	{
+		Unknown,
+		VP8,
+		VP9
+	};
+
 	vpx_codec_ctx_t Context;
+	FCriticalSection PreviousContextsLock;
+	TArray<vpx_codec_ctx_t> PreviousContexts;
+	vpx_codec_dec_cfg_t CodecConfig {};
 	TUniquePtr<FWebMMediaTextureSamplePool> VideoSamplePool;
 	TRefCountPtr<FRHITexture> DecodedY;
 	TRefCountPtr<FRHITexture> DecodedU;
 	TRefCountPtr<FRHITexture> DecodedV;
 	FGraphEventRef VideoDecodingTask;
 	IWebMSamplesSink& Samples;
+	ECodec Codec;
+	int64 CurrentSequenceIndex;
 	bool bTexturesCreated;
 	bool bIsInitialized;
 
+	bool InitializeInternal(FString InCodecName);
 	void ConvertYUVToRGBAndSubmit(FRHICommandListImmediate& RHICmdList, const FConvertParams& Params);
-	void DoDecodeVideoFrames(const TArray<TSharedPtr<FWebMFrame>>& VideoFrames);
+	void DoDecodeVideoFrames(const TArray<TSharedPtr<FWebMFrame, ESPMode::ThreadSafe>>& VideoFrames);
 	void CreateTextures(FRHICommandListBase& RHICmdList, const vpx_image_t* Image);
 	void Close();
+	void ResetPendingDecoders();
 };
 
 #endif // WITH_WEBM_LIBS
diff --git a/Engine/Plugins/Runtime/WebMMoviePlayer/Source/WebMMoviePlayer/Private/WebMMovieStreamer.cpp b/Engine/Plugins/Runtime/WebMMoviePlayer/Source/WebMMoviePlayer/Private/WebMMovieStreamer.cpp
--- a/Engine/Plugins/Runtime/WebMMoviePlayer/Source/WebMMoviePlayer/Private/WebMMovieStreamer.cpp
+++ b/Engine/Plugins/Runtime/WebMMoviePlayer/Source/WebMMoviePlayer/Private/WebMMovieStreamer.cpp
@@ -106,7 +106,7 @@ bool FWebMMovieStreamer::StartNextMovie()
 		else
 		{
 			UE_LOG(LogWebMMoviePlayer, Error, TEXT("Movie '%s' not found."), *MoviePath);
-			
+
 			MovieName = FString();
 			return false;
 		}
@@ -269,7 +269,7 @@ bool FWebMMovieStreamer::DisplayFrames(float InDeltaTime)
 bool FWebMMovieStreamer::SendAudio(float InDeltaTime)
 {
 	// Just send all available audio for processing
-	
+
 	TSharedPtr<IMediaAudioSample, ESPMode::ThreadSafe> AudioSample;
 	TRange<FTimespan> TimeRange(FTimespan::Zero(), FTimespan::MaxValue());
 	bool bFoundSample = Samples->FetchAudio(TimeRange, AudioSample);
@@ -289,8 +289,8 @@ bool FWebMMovieStreamer::ReadMoreFrames()
 {
 	FTimespan ReadBufferLength = FTimespan::FromSeconds(1.0 / 30.0);
 
-	TArray<TSharedPtr<FWebMFrame>> AudioFrames;
-	TArray<TSharedPtr<FWebMFrame>> VideoFrames;
+	TArray<TSharedPtr<FWebMFrame, ESPMode::ThreadSafe>> AudioFrames;
+	TArray<TSharedPtr<FWebMFrame, ESPMode::ThreadSafe>> VideoFrames;
 
 	Container->ReadFrames(ReadBufferLength, AudioFrames, VideoFrames);
 
diff --git a/Engine/Plugins/Runtime/WebMMoviePlayer/Source/WebMMoviePlayer/Private/WebMMovieStreamer.h b/Engine/Plugins/Runtime/WebMMoviePlayer/Source/WebMMoviePlayer/Private/WebMMovieStreamer.h
--- a/Engine/Plugins/Runtime/WebMMoviePlayer/Source/WebMMoviePlayer/Private/WebMMovieStreamer.h
+++ b/Engine/Plugins/Runtime/WebMMoviePlayer/Source/WebMMoviePlayer/Private/WebMMovieStreamer.h
@@ -26,6 +26,8 @@ public:
 	//~ IWebMSamplesSink interface
 	virtual void AddVideoSampleFromDecodingThread(TSharedRef<FWebMMediaTextureSample, ESPMode::ThreadSafe> Sample) override;
 	virtual void AddAudioSampleFromDecodingThread(TSharedRef<FWebMMediaAudioSample, ESPMode::ThreadSafe> Sample) override;
+	virtual void ReportVideoDecodingError(FString InErrorMessage) override { }
+	virtual void ReportAudioDecodingError(FString InErrorMessage) override { }
 
 	//~ IMediaOutput interface
 	virtual bool Init(const TArray<FString>& InMoviePaths, TEnumAsByte<EMoviePlaybackType> InPlaybackType) override;
@@ -43,7 +45,7 @@ public:
 
 private:
 	TArray<FString> MovieQueue;
-	TQueue<TArray<TSharedPtr<FWebMFrame>>> VideoFramesToDecodeLater;
+	TQueue<TArray<TSharedPtr<FWebMFrame, ESPMode::ThreadSafe>>> VideoFramesToDecodeLater;
 	FString MovieName;
 	TUniquePtr<FWebMVideoDecoder> VideoDecoder;
 	TUniquePtr<FWebMAudioDecoder> AudioDecoder;
@@ -65,9 +67,9 @@ private:
 	/** We can be paused e.g. while in background */
 	bool bPaused;
 
-	/** 
-	 * Number of ticks to wait after the movie is complete before moving on to the next one. 
-	 * 
+	/**
+	 * Number of ticks to wait after the movie is complete before moving on to the next one.
+	 *
 	 * This allows us to defer texture deletion while it is being displayed.
 	 */
 	int32 TicksLeftToWaitPostCompletion;
